Another fix for incorrect function name.

diff --git a/Engine/Plugins/Runtime/CableComponent/Source/CableComponent/Private/CableComponent.cpp b/Engine/Plugins/Runtime/CableComponent/Source/CableComponent/Private/CableComponent.cpp
--- a/Engine/Plugins/Runtime/CableComponent/Source/CableComponent/Private/CableComponent.cpp
+++ b/Engine/Plugins/Runtime/CableComponent/Source/CableComponent/Private/CableComponent.cpp
@@ -243,7 +243,7 @@ public:
 	{
 		QUICK_SCOPE_CYCLE_COUNTER( STAT_CableSceneProxy_DrawDynamicElements );
 
-		const bool bWireframe = GEngine->AllowDebugViewModes() && View->Family->EngineShowFlags.Wireframe;
+		const bool bWireframe = AllowDebugViewmodes() && View->Family->EngineShowFlags.Wireframe;
 
 		FColoredMaterialRenderProxy WireframeMaterialInstance(
 			GEngine->WireframeMaterial ? GEngine->WireframeMaterial->GetRenderProxy(IsSelected()) : NULL,
