Expose LOD settings for CableComponent
Also fix bounding box (and drawing of it)

diff --git a/Engine/Plugins/Runtime/CableComponent/Source/CableComponent/Classes/CableComponent.h b/Engine/Plugins/Runtime/CableComponent/Source/CableComponent/Classes/CableComponent.h
--- a/Engine/Plugins/Runtime/CableComponent/Source/CableComponent/Classes/CableComponent.h
+++ b/Engine/Plugins/Runtime/CableComponent/Source/CableComponent/Classes/CableComponent.h
@@ -22,7 +22,7 @@ struct FCableParticle
 };
 
 /** Component that allows you to specify custom triangle mesh geometry */
-UCLASS(hidecategories=(Object, LOD, Physics, Collision, Activation, "Components|Activation"), editinlinenew, meta=(BlueprintSpawnableComponent), ClassGroup=Rendering)
+UCLASS(hidecategories=(Object, Physics, Collision, Activation, "Components|Activation"), editinlinenew, meta=(BlueprintSpawnableComponent), ClassGroup=Rendering)
 class UCableComponent : public UMeshComponent
 {
 	GENERATED_UCLASS_BODY()
diff --git a/Engine/Plugins/Runtime/CableComponent/Source/CableComponent/Private/CableComponent.cpp b/Engine/Plugins/Runtime/CableComponent/Source/CableComponent/Private/CableComponent.cpp
--- a/Engine/Plugins/Runtime/CableComponent/Source/CableComponent/Private/CableComponent.cpp
+++ b/Engine/Plugins/Runtime/CableComponent/Source/CableComponent/Private/CableComponent.cpp
@@ -279,6 +279,11 @@ public:
 		Mesh.Type = PT_TriangleList;
 		Mesh.DepthPriorityGroup = SDPG_World;
 		PDI->DrawMesh(Mesh);
+
+#if !(UE_BUILD_SHIPPING || UE_BUILD_TEST)
+		// Render bounds
+		RenderBounds(PDI, View->Family->EngineShowFlags, GetBounds(), IsSelected());
+#endif
 	}
 
 	virtual FPrimitiveViewRelevance GetViewRelevance(const FSceneView* View)
@@ -538,8 +543,6 @@ FBoxSphereBounds UCableComponent::CalcBounds(const FTransform & LocalToWorld) co
 	}
 
 	// Expand by cable width
-	CableBox.ExpandBy(CableWidth);
-
-	return FBoxSphereBounds(CableBox);
+	return FBoxSphereBounds(CableBox.ExpandBy(CableWidth));
 }
 
