Added BP visible functions for getting and setting the AttachEndTo position in CableComponent.
#github Pull Request #417

diff --git a/Engine/Plugins/Runtime/CableComponent/Source/CableComponent/Classes/CableComponent.h b/Engine/Plugins/Runtime/CableComponent/Source/CableComponent/Classes/CableComponent.h
--- a/Engine/Plugins/Runtime/CableComponent/Source/CableComponent/Classes/CableComponent.h
+++ b/Engine/Plugins/Runtime/CableComponent/Source/CableComponent/Classes/CableComponent.h
@@ -53,9 +53,21 @@ public:
 	FComponentReference AttachEndTo;
 
 	/** End location of cable, relative to AttachEndTo if specified, otherwise relative to cable component. */
-	UPROPERTY(EditAnywhere, Category="Cable")
+	UPROPERTY(EditAnywhere, BlueprintReadWrite, Category = "Cable")
 	FVector EndLocation;
 
+	/** Attaches the end of the cable to a specific Component within an Actor **/
+	UFUNCTION(BlueprintCallable, Category = "Cable")
+	void SetAttachEndTo(AActor* Actor, FName ComponentProperty);
+	
+	/** Gets the Actor that the cable is attached to **/
+	UFUNCTION(BlueprintCallable, Category = "Cable")
+	AActor* GetAttachedActor() const;
+
+	/** Gets the specific USceneComponent that the cable is attached to **/
+	UFUNCTION(BlueprintCallable, Category = "Cable")
+	USceneComponent* GetAttachedComponent() const;
+
 	/** Rest length of the cable */
 	UPROPERTY(EditAnywhere, BlueprintReadWrite, Category="Cable", meta=(ClampMin = "0.0", UIMin = "0.0", UIMax = "1000.0"))
 	float CableLength;
diff --git a/Engine/Plugins/Runtime/CableComponent/Source/CableComponent/Private/CableComponent.cpp b/Engine/Plugins/Runtime/CableComponent/Source/CableComponent/Private/CableComponent.cpp
--- a/Engine/Plugins/Runtime/CableComponent/Source/CableComponent/Private/CableComponent.cpp
+++ b/Engine/Plugins/Runtime/CableComponent/Source/CableComponent/Private/CableComponent.cpp
@@ -467,6 +467,22 @@ void UCableComponent::PerformSubstep(float InSubstepTime, const FVector& Gravity
 	SolveConstraints();
 }
 
+void UCableComponent::SetAttachEndTo(AActor* Actor, FName ComponentProperty)
+{
+	AttachEndTo.OtherActor = Actor;
+	AttachEndTo.ComponentProperty = ComponentProperty;
+}
+
+AActor* UCableComponent::GetAttachedActor() const
+{
+	return AttachEndTo.OtherActor;
+}
+
+USceneComponent* UCableComponent::GetAttachedComponent() const
+{
+	return AttachEndTo.GetComponent(GetOwner());
+}
+
 void UCableComponent::GetEndPositions(FVector& OutStartPosition, FVector& OutEndPosition)
 {
 	// Start position is just component position
