UE-15620 Add CABLECOMPONENT_API to cable component class

diff --git a/Engine/Plugins/Runtime/CableComponent/Source/CableComponent/Classes/CableComponent.h b/Engine/Plugins/Runtime/CableComponent/Source/CableComponent/Classes/CableComponent.h
--- a/Engine/Plugins/Runtime/CableComponent/Source/CableComponent/Classes/CableComponent.h
+++ b/Engine/Plugins/Runtime/CableComponent/Source/CableComponent/Classes/CableComponent.h
@@ -23,7 +23,7 @@ struct FCableParticle
 
 /** Component that allows you to specify custom triangle mesh geometry */
 UCLASS(hidecategories=(Object, Physics, Collision, Activation, "Components|Activation"), editinlinenew, meta=(BlueprintSpawnableComponent), ClassGroup=Rendering)
-class UCableComponent : public UMeshComponent
+class CABLECOMPONENT_API UCableComponent : public UMeshComponent
 {
 	GENERATED_UCLASS_BODY()
 
