Fix comments.

diff --git a/Engine/Plugins/Runtime/CableComponent/Source/CableComponent/Classes/CableComponent.h b/Engine/Plugins/Runtime/CableComponent/Source/CableComponent/Classes/CableComponent.h
--- a/Engine/Plugins/Runtime/CableComponent/Source/CableComponent/Classes/CableComponent.h
+++ b/Engine/Plugins/Runtime/CableComponent/Source/CableComponent/Classes/CableComponent.h
@@ -41,7 +41,7 @@ public:
 	virtual void SendRenderDynamicData_Concurrent() override;
 	virtual void CreateRenderState_Concurrent() override;
 	virtual void ApplyWorldOffset(const FVector& InOffset, bool bWorldShift) override;
-	//~ Begin UActorComponent Interface.
+	//~ End UActorComponent Interface.
 
 	//~ Begin USceneComponent Interface.
 	virtual FBoxSphereBounds CalcBounds(const FTransform& LocalToWorld) const override;
@@ -49,7 +49,7 @@ public:
 	virtual bool HasAnySockets() const override;
 	virtual bool DoesSocketExist(FName InSocketName) const override;
 	virtual FTransform GetSocketTransform(FName InSocketName, ERelativeTransformSpace TransformSpace = RTS_World) const override;
-	//~ Begin USceneComponent Interface.
+	//~ End USceneComponent Interface.
 
 	//~ Begin UPrimitiveComponent Interface.
 	virtual FPrimitiveSceneProxy* CreateSceneProxy() override;
