Add an option for attaching a cable end directly to a scene component, when the actor/propertyname are not known.

diff --git a/Engine/Plugins/Runtime/CableComponent/Source/CableComponent/Classes/CableComponent.h b/Engine/Plugins/Runtime/CableComponent/Source/CableComponent/Classes/CableComponent.h
--- a/Engine/Plugins/Runtime/CableComponent/Source/CableComponent/Classes/CableComponent.h
+++ b/Engine/Plugins/Runtime/CableComponent/Source/CableComponent/Classes/CableComponent.h
@@ -86,6 +86,9 @@ public:
 	UPROPERTY(EditAnywhere, BlueprintReadWrite, Category="Cable", meta=(MakeEditWidget=true))
 	FVector EndLocation;
 
+	/** Attaches the end of the cable to a specific Component **/
+	void SetAttachEndToComponent(TWeakObjectPtr<class USceneComponent> Component, FName SocketName = NAME_None);
+
 	/** Attaches the end of the cable to a specific Component within an Actor **/
 	UFUNCTION(BlueprintCallable, Category = "Cable")
 	void SetAttachEndTo(AActor* Actor, FName ComponentProperty, FName SocketName = NAME_None);
diff --git a/Engine/Plugins/Runtime/CableComponent/Source/CableComponent/Private/CableComponent.cpp b/Engine/Plugins/Runtime/CableComponent/Source/CableComponent/Private/CableComponent.cpp
--- a/Engine/Plugins/Runtime/CableComponent/Source/CableComponent/Private/CableComponent.cpp
+++ b/Engine/Plugins/Runtime/CableComponent/Source/CableComponent/Private/CableComponent.cpp
@@ -555,6 +555,13 @@ void UCableComponent::PerformSubstep(float InSubstepTime, const FVector& Gravity
 	}
 }
 
+void UCableComponent::SetAttachEndToComponent(TWeakObjectPtr<class USceneComponent> Component, FName SocketName)
+{
+	AttachEndTo.OtherActor = Component->GetOwner();
+	AttachEndTo.OverrideComponent = Component;
+	AttachEndToSocketName = SocketName;
+}
+
 void UCableComponent::SetAttachEndTo(AActor* Actor, FName ComponentProperty, FName SocketName)
 {
 	AttachEndTo.OtherActor = Actor;
