Pass USceneComponent* to CableComponent instead of TWeakObjectPtr.
Expose AttachEndsToComponent to blueprints.
Set ComponentProperty name to NAME_None when component override is used.

diff --git a/Engine/Plugins/Runtime/CableComponent/Source/CableComponent/Classes/CableComponent.h b/Engine/Plugins/Runtime/CableComponent/Source/CableComponent/Classes/CableComponent.h
--- a/Engine/Plugins/Runtime/CableComponent/Source/CableComponent/Classes/CableComponent.h
+++ b/Engine/Plugins/Runtime/CableComponent/Source/CableComponent/Classes/CableComponent.h
@@ -87,7 +87,8 @@ public:
 	FVector EndLocation;
 
 	/** Attaches the end of the cable to a specific Component **/
-	void SetAttachEndToComponent(TWeakObjectPtr<class USceneComponent> Component, FName SocketName = NAME_None);
+	UFUNCTION(BlueprintCallable, Category = "Cable")
+	void SetAttachEndToComponent(USceneComponent* Component, FName SocketName = NAME_None);
 
 	/** Attaches the end of the cable to a specific Component within an Actor **/
 	UFUNCTION(BlueprintCallable, Category = "Cable")
diff --git a/Engine/Plugins/Runtime/CableComponent/Source/CableComponent/Private/CableComponent.cpp b/Engine/Plugins/Runtime/CableComponent/Source/CableComponent/Private/CableComponent.cpp
--- a/Engine/Plugins/Runtime/CableComponent/Source/CableComponent/Private/CableComponent.cpp
+++ b/Engine/Plugins/Runtime/CableComponent/Source/CableComponent/Private/CableComponent.cpp
@@ -555,9 +555,10 @@ void UCableComponent::PerformSubstep(float InSubstepTime, const FVector& Gravity
 	}
 }
 
-void UCableComponent::SetAttachEndToComponent(TWeakObjectPtr<class USceneComponent> Component, FName SocketName)
+void UCableComponent::SetAttachEndToComponent(USceneComponent* Component, FName SocketName)
 {
 	AttachEndTo.OtherActor = Component->GetOwner();
+	AttachEndTo.ComponentProperty = NAME_None;
 	AttachEndTo.OverrideComponent = Component;
 	AttachEndToSocketName = SocketName;
 }
