Fix cable component bounds calculation by transforming points to local space.

diff --git a/Engine/Plugins/Runtime/CableComponent/Source/CableComponent/Private/CableComponent.cpp b/Engine/Plugins/Runtime/CableComponent/Source/CableComponent/Private/CableComponent.cpp
--- a/Engine/Plugins/Runtime/CableComponent/Source/CableComponent/Private/CableComponent.cpp
+++ b/Engine/Plugins/Runtime/CableComponent/Source/CableComponent/Private/CableComponent.cpp
@@ -757,14 +757,17 @@ FBoxSphereBounds UCableComponent::CalcBounds(const FTransform& LocalToWorld) con
 {
 	// Calculate bounding box of cable points
 	FBox CableBox(ForceInit);
+
+	const FTransform& ComponentTransform = GetComponentTransform();
+
 	for(int32 ParticleIdx=0; ParticleIdx<Particles.Num(); ParticleIdx++)
 	{
 		const FCableParticle& Particle = Particles[ParticleIdx];
-		CableBox += Particle.Position;
+		CableBox += ComponentTransform.InverseTransformPosition(Particle.Position);
 	}
 
 	// Expand by cable radius (half cable width)
-	return FBoxSphereBounds(CableBox.ExpandBy(0.5f * CableWidth));
+	return FBoxSphereBounds(CableBox.ExpandBy(0.5f * CableWidth)).TransformBy(LocalToWorld);
 }
 
 void UCableComponent::QuerySupportedSockets(TArray<FComponentSocketDescription>& OutSockets) const 
