Fixed flickering due a race condition in the cable component

diff --git a/Engine/Plugins/Runtime/CableComponent/Source/CableComponent/Private/CableComponent.cpp b/Engine/Plugins/Runtime/CableComponent/Source/CableComponent/Private/CableComponent.cpp
--- a/Engine/Plugins/Runtime/CableComponent/Source/CableComponent/Private/CableComponent.cpp
+++ b/Engine/Plugins/Runtime/CableComponent/Source/CableComponent/Private/CableComponent.cpp
@@ -64,7 +64,6 @@ public:
 		: FPrimitiveSceneProxy(Component)
 		, Material(NULL)
 		, VertexFactory(GetScene().GetFeatureLevel(), "FCableSceneProxy")
-		, DynamicData(NULL)
 		, MaterialRelevance(Component->GetMaterialRelevance(GetScene().GetFeatureLevel()))
 		, NumSegments(Component->NumSegments)
 		, CableWidth(Component->CableWidth)
@@ -93,11 +92,6 @@ public:
 		VertexBuffers.ColorVertexBuffer.ReleaseResource();
 		IndexBuffer.ReleaseResource();
 		VertexFactory.ReleaseResource();
-
-		if(DynamicData != NULL)
-		{
-			delete DynamicData;
-		}
 	}
 
 	int32 GetRequiredVertexCount() const
@@ -187,63 +181,61 @@ public:
 	{
 		check(IsInRenderingThread());
 
-		// Free existing data if present
-		if(DynamicData)
+		if (NewDynamicData != nullptr)
 		{
-			delete DynamicData;
-			DynamicData = NULL;
-		}
-		DynamicData = NewDynamicData;
+			// Build mesh from cable points
+			TArray<FDynamicMeshVertex> Vertices;
+			TArray<int32> Indices;
+			BuildCableMesh(NewDynamicData->CablePoints, Vertices, Indices);
 
-		// Build mesh from cable points
-		TArray<FDynamicMeshVertex> Vertices;
-		TArray<int32> Indices;
-		BuildCableMesh(NewDynamicData->CablePoints, Vertices, Indices);
+			check(Vertices.Num() == GetRequiredVertexCount());
+			check(Indices.Num() == GetRequiredIndexCount());
 
-		check(Vertices.Num() == GetRequiredVertexCount());
-		check(Indices.Num() == GetRequiredIndexCount());
+			for (int i = 0; i < Vertices.Num(); i++)
+			{
+				const FDynamicMeshVertex& Vertex = Vertices[i];
 
-		for (int i = 0; i < Vertices.Num(); i++)
-		{
-			const FDynamicMeshVertex& Vertex = Vertices[i];
+				VertexBuffers.PositionVertexBuffer.VertexPosition(i) = Vertex.Position;
+				VertexBuffers.StaticMeshVertexBuffer.SetVertexTangents(i, Vertex.TangentX.ToFVector(), Vertex.GetTangentY(), Vertex.TangentZ.ToFVector());
+				VertexBuffers.StaticMeshVertexBuffer.SetVertexUV(i, 0, Vertex.TextureCoordinate[0]);
+				VertexBuffers.ColorVertexBuffer.VertexColor(i) = Vertex.Color;
+			}
 
-			VertexBuffers.PositionVertexBuffer.VertexPosition(i) = Vertex.Position;
-			VertexBuffers.StaticMeshVertexBuffer.SetVertexTangents(i, Vertex.TangentX.ToFVector(), Vertex.GetTangentY(), Vertex.TangentZ.ToFVector());
-			VertexBuffers.StaticMeshVertexBuffer.SetVertexUV(i, 0, Vertex.TextureCoordinate[0]);
-			VertexBuffers.ColorVertexBuffer.VertexColor(i) = Vertex.Color;
-		}
+			{
+				auto& VertexBuffer = VertexBuffers.PositionVertexBuffer;
+				void* VertexBufferData = RHILockBuffer(VertexBuffer.VertexBufferRHI, 0, VertexBuffer.GetNumVertices() * VertexBuffer.GetStride(), RLM_WriteOnly);
+				FMemory::Memcpy(VertexBufferData, VertexBuffer.GetVertexData(), VertexBuffer.GetNumVertices() * VertexBuffer.GetStride());
+				RHIUnlockBuffer(VertexBuffer.VertexBufferRHI);
+			}
 
-		{
-			auto& VertexBuffer = VertexBuffers.PositionVertexBuffer;
-			void* VertexBufferData = RHILockBuffer(VertexBuffer.VertexBufferRHI, 0, VertexBuffer.GetNumVertices() * VertexBuffer.GetStride(), RLM_WriteOnly);
-			FMemory::Memcpy(VertexBufferData, VertexBuffer.GetVertexData(), VertexBuffer.GetNumVertices() * VertexBuffer.GetStride());
-			RHIUnlockBuffer(VertexBuffer.VertexBufferRHI);
-		}
+			{
+				auto& VertexBuffer = VertexBuffers.ColorVertexBuffer;
+				void* VertexBufferData = RHILockBuffer(VertexBuffer.VertexBufferRHI, 0, VertexBuffer.GetNumVertices() * VertexBuffer.GetStride(), RLM_WriteOnly);
+				FMemory::Memcpy(VertexBufferData, VertexBuffer.GetVertexData(), VertexBuffer.GetNumVertices() * VertexBuffer.GetStride());
+				RHIUnlockBuffer(VertexBuffer.VertexBufferRHI);
+			}
 
-		{
-			auto& VertexBuffer = VertexBuffers.ColorVertexBuffer;
-			void* VertexBufferData = RHILockBuffer(VertexBuffer.VertexBufferRHI, 0, VertexBuffer.GetNumVertices() * VertexBuffer.GetStride(), RLM_WriteOnly);
-			FMemory::Memcpy(VertexBufferData, VertexBuffer.GetVertexData(), VertexBuffer.GetNumVertices() * VertexBuffer.GetStride());
-			RHIUnlockBuffer(VertexBuffer.VertexBufferRHI);
-		}
+			{
+				auto& VertexBuffer = VertexBuffers.StaticMeshVertexBuffer;
+				void* VertexBufferData = RHILockBuffer(VertexBuffer.TangentsVertexBuffer.VertexBufferRHI, 0, VertexBuffer.GetTangentSize(), RLM_WriteOnly);
+				FMemory::Memcpy(VertexBufferData, VertexBuffer.GetTangentData(), VertexBuffer.GetTangentSize());
+				RHIUnlockBuffer(VertexBuffer.TangentsVertexBuffer.VertexBufferRHI);
+			}
 
-		{
-			auto& VertexBuffer = VertexBuffers.StaticMeshVertexBuffer;
-			void* VertexBufferData = RHILockBuffer(VertexBuffer.TangentsVertexBuffer.VertexBufferRHI, 0, VertexBuffer.GetTangentSize(), RLM_WriteOnly);
-			FMemory::Memcpy(VertexBufferData, VertexBuffer.GetTangentData(), VertexBuffer.GetTangentSize());
-			RHIUnlockBuffer(VertexBuffer.TangentsVertexBuffer.VertexBufferRHI);
-		}
+			{
+				auto& VertexBuffer = VertexBuffers.StaticMeshVertexBuffer;
+				void* VertexBufferData = RHILockBuffer(VertexBuffer.TexCoordVertexBuffer.VertexBufferRHI, 0, VertexBuffer.GetTexCoordSize(), RLM_WriteOnly);
+				FMemory::Memcpy(VertexBufferData, VertexBuffer.GetTexCoordData(), VertexBuffer.GetTexCoordSize());
+				RHIUnlockBuffer(VertexBuffer.TexCoordVertexBuffer.VertexBufferRHI);
+			}
 
-		{
-			auto& VertexBuffer = VertexBuffers.StaticMeshVertexBuffer;
-			void* VertexBufferData = RHILockBuffer(VertexBuffer.TexCoordVertexBuffer.VertexBufferRHI, 0, VertexBuffer.GetTexCoordSize(), RLM_WriteOnly);
-			FMemory::Memcpy(VertexBufferData, VertexBuffer.GetTexCoordData(), VertexBuffer.GetTexCoordSize());
-			RHIUnlockBuffer(VertexBuffer.TexCoordVertexBuffer.VertexBufferRHI);
-		}
+			void* IndexBufferData = RHILockBuffer(IndexBuffer.IndexBufferRHI, 0, Indices.Num() * sizeof(int32), RLM_WriteOnly);
+			FMemory::Memcpy(IndexBufferData, &Indices[0], Indices.Num() * sizeof(int32));
+			RHIUnlockBuffer(IndexBuffer.IndexBufferRHI);
 
-		void* IndexBufferData = RHILockBuffer(IndexBuffer.IndexBufferRHI, 0, Indices.Num() * sizeof(int32), RLM_WriteOnly);
-		FMemory::Memcpy(IndexBufferData, &Indices[0], Indices.Num() * sizeof(int32));
-		RHIUnlockBuffer(IndexBuffer.IndexBufferRHI);
+			delete NewDynamicData;
+			NewDynamicData = NULL;
+		}
 	}
 
 	virtual void GetDynamicMeshElements(const TArray<const FSceneView*>& Views, const FSceneViewFamily& ViewFamily, uint32 VisibilityMap, FMeshElementCollector& Collector) const override
@@ -331,8 +323,6 @@ private:
 	FCableIndexBuffer IndexBuffer;
 	FLocalVertexFactory VertexFactory;
 
-	FCableDynamicData* DynamicData;
-
 	FMaterialRelevance MaterialRelevance;
 
 	int32 NumSegments;
