Fixed cables not rendering depth correctly when dynamically moved around in scene (in editor of via blueprint).
This matches other bVelocityRelevance setup and behaves the same as in 4.27.

diff --git a/Engine/Plugins/Runtime/CableComponent/Source/CableComponent/Private/CableComponent.cpp b/Engine/Plugins/Runtime/CableComponent/Source/CableComponent/Private/CableComponent.cpp
--- a/Engine/Plugins/Runtime/CableComponent/Source/CableComponent/Private/CableComponent.cpp
+++ b/Engine/Plugins/Runtime/CableComponent/Source/CableComponent/Private/CableComponent.cpp
@@ -308,7 +308,11 @@ public:
 		Result.bDrawRelevance = IsShown(View);
 		Result.bShadowRelevance = IsShadowCast(View);
 		Result.bDynamicRelevance = true;
+
 		MaterialRelevance.SetPrimitiveViewRelevance(Result);
+
+		Result.bVelocityRelevance = DrawsVelocity() && Result.bOpaque && Result.bRenderInMainPass;
+
 		return Result;
 	}
 
