PR #9630: Allow UCableComponent to render to the CustomDepth buffer (Contributed by [REDACTED_PERSON])

diff --git a/Engine/Plugins/Runtime/CableComponent/Source/CableComponent/Private/CableComponent.cpp b/Engine/Plugins/Runtime/CableComponent/Source/CableComponent/Private/CableComponent.cpp
--- a/Engine/Plugins/Runtime/CableComponent/Source/CableComponent/Private/CableComponent.cpp
+++ b/Engine/Plugins/Runtime/CableComponent/Source/CableComponent/Private/CableComponent.cpp
@@ -310,6 +310,10 @@ public:
 		FPrimitiveViewRelevance Result;
 		Result.bDrawRelevance = IsShown(View);
 		Result.bShadowRelevance = IsShadowCast(View);
+		Result.bRenderCustomDepth = ShouldRenderCustomDepth();
+		Result.bRenderInMainPass = ShouldRenderInMainPass();
+		Result.bUsesLightingChannels = GetLightingChannelMask() != GetDefaultLightingChannelMask();
+		Result.bTranslucentSelfShadow = bCastVolumetricTranslucentShadow;
 		Result.bDynamicRelevance = true;
 
 		MaterialRelevance.SetPrimitiveViewRelevance(Result);
