Cable Component: get around the fact that it keeps an array of world space positions, resulting in bad actor bounds when saving a streaming level.

#rnx

diff --git a/Engine/Plugins/Runtime/CableComponent/Source/CableComponent/Classes/CableComponent.h b/Engine/Plugins/Runtime/CableComponent/Source/CableComponent/Classes/CableComponent.h
--- a/Engine/Plugins/Runtime/CableComponent/Source/CableComponent/Classes/CableComponent.h
+++ b/Engine/Plugins/Runtime/CableComponent/Source/CableComponent/Classes/CableComponent.h
@@ -60,6 +60,9 @@ public:
 	virtual int32 GetNumMaterials() const override;
 	//~ End UMeshComponent Interface.
 
+#if WITH_EDITOR
+	virtual void OnUpdateTransform(EUpdateTransformFlags UpdateTransformFlags, ETeleportType Teleport) override;
+#endif
 
 	/**
 	 *	Should we fix the start to something, or leave it free.
@@ -171,6 +174,8 @@ public:
 
 private:
 
+	/** Initialize the particles */
+	void InitParticles();
 	/** Solve the cable spring constraints */
 	void SolveConstraints();
 	/** Integrate cable point positions */
diff --git a/Engine/Plugins/Runtime/CableComponent/Source/CableComponent/Private/CableComponent.cpp b/Engine/Plugins/Runtime/CableComponent/Source/CableComponent/Private/CableComponent.cpp
--- a/Engine/Plugins/Runtime/CableComponent/Source/CableComponent/Private/CableComponent.cpp
+++ b/Engine/Plugins/Runtime/CableComponent/Source/CableComponent/Private/CableComponent.cpp
@@ -370,6 +370,10 @@ UCableComponent::UCableComponent( const FObjectInitializer& ObjectInitializer )
 	CableGravityScale = 1.f;
 
 	SetCollisionProfileName(UCollisionProfile::PhysicsActor_ProfileName);
+
+#if WITH_EDITOR
+	bWantsOnUpdateTransform = true;
+#endif
 }
 
 FPrimitiveSceneProxy* UCableComponent::CreateSceneProxy()
@@ -385,29 +389,19 @@ int32 UCableComponent::GetNumMaterials() const
 void UCableComponent::OnRegister()
 {
 	Super::OnRegister();
+	
+	InitParticles();
+}
 
-	const int32 NumParticles = NumSegments+1;
-
-	Particles.Reset();
-	Particles.AddUninitialized(NumParticles);
-
-	FVector CableStart, CableEnd;
-	GetEndPositions(CableStart, CableEnd);
-
-	const FVector Delta = CableEnd - CableStart;
-
-	for(int32 ParticleIdx=0; ParticleIdx<NumParticles; ParticleIdx++)
-	{
-		FCableParticle& Particle = Particles[ParticleIdx];
-
-		const float Alpha = (float)ParticleIdx/(float)NumSegments;
-		const FVector InitialPosition = CableStart + (Alpha * Delta);
+#if WITH_EDITOR
+void UCableComponent::OnUpdateTransform(EUpdateTransformFlags UpdateTransformFlags, ETeleportType Teleport)
+{
+	Super::OnUpdateTransform(UpdateTransformFlags, Teleport);
 
-		Particle.Position = InitialPosition;
-		Particle.OldPosition = InitialPosition;
-		Particle.bFree = true; // default to free, will be fixed if desired in TickComponent
-	}
+	InitParticles();
+	UpdateBounds();
 }
+#endif
 
 void UCableComponent::VerletIntegrate(float InSubstepTime, const FVector& Gravity)
 {
@@ -464,6 +458,31 @@ static FORCEINLINE void SolveDistanceConstraint(FCableParticle& ParticleA, FCabl
 	}
 }
 
+void UCableComponent::InitParticles()
+{
+	const int32 NumParticles = NumSegments+1;
+
+	Particles.Reset();
+	Particles.AddUninitialized(NumParticles);
+
+	FVector CableStart, CableEnd;
+	GetEndPositions(CableStart, CableEnd);
+
+	const FVector Delta = CableEnd - CableStart;
+
+	for(int32 ParticleIdx=0; ParticleIdx<NumParticles; ParticleIdx++)
+	{
+		FCableParticle& Particle = Particles[ParticleIdx];
+
+		const float Alpha = (float)ParticleIdx/(float)NumSegments;
+		const FVector InitialPosition = CableStart + (Alpha * Delta);
+
+		Particle.Position = InitialPosition;
+		Particle.OldPosition = InitialPosition;
+		Particle.bFree = true; // default to free, will be fixed if desired in TickComponent
+	}
+}
+
 void UCableComponent::SolveConstraints()
 {
 	SCOPE_CYCLE_COUNTER(STAT_Cable_SolveTime);
