Add ray tracing support to FCableSceneProxyCable (with WPO).

diff --git a/Engine/Plugins/Runtime/CableComponent/Source/CableComponent/Private/CableComponent.cpp b/Engine/Plugins/Runtime/CableComponent/Source/CableComponent/Private/CableComponent.cpp
--- a/Engine/Plugins/Runtime/CableComponent/Source/CableComponent/Private/CableComponent.cpp
+++ b/Engine/Plugins/Runtime/CableComponent/Source/CableComponent/Private/CableComponent.cpp
@@ -14,6 +14,30 @@
 #include "StaticMeshResources.h"
 #include "SceneInterface.h"
 
+
+static TAutoConsoleVariable<int32> CVarRayTracingCableMeshes(
+	TEXT("r.RayTracing.Geometry.Cable"),
+	1,
+	TEXT("Include Cable meshes in ray tracing effects (default = 1 (cable meshes enabled in ray tracing))"));
+
+static TAutoConsoleVariable<int32> CVarRayTracingCableMeshesWPO(
+	TEXT("r.RayTracing.Geometry.Cable.WPO"),
+	1,
+	TEXT("World position offset evaluation for cable meshes with EvaluateWPO enabled in ray tracing effects.\n")
+	TEXT(" 0: Cable meshes with world position offset visible in ray tracing, WPO evaluation disabled.\n")
+	TEXT(" 1: Cable meshes with world position offset visible in ray tracing, WPO evaluation enabled (default).\n")
+);
+
+static TAutoConsoleVariable<int32> CVarRayTracingCableMeshesWPOCulling(
+	TEXT("r.RayTracing.Geometry.Cable.WPO.Culling"),
+	1,
+	TEXT("Enable culling for WPO evaluation for cable meshes in ray tracing (default = 1 (Culling enabled))"));
+
+static TAutoConsoleVariable<float> CVarRayTracingCableMeshesWPOCullingRadius(
+	TEXT("r.RayTracing.Geometry.Cable.WPO.CullingRadius"),
+	12000.0f, // 120 m
+	TEXT("Do not evaluate world position offset for cable meshes outside of this radius in ray tracing effects (default = 12000 (120m))"));
+
 #include UE_INLINE_GENERATED_CPP_BY_NAME(CableComponent)
 
 DECLARE_CYCLE_STAT(TEXT("Cable Sim"), STAT_Cable_SimTime, STATGROUP_CableComponent);
@@ -81,6 +105,26 @@ public:
 		{
 			Material = UMaterial::GetDefaultMaterial(MD_Surface);
 		}
+
+#if RHI_RAYTRACING
+		bSupportRayTracing = IsRayTracingEnabled();
+		bDynamicRayTracingGeometry = false;
+		bNeedsDynamicRayTracingGeometries = false;
+		
+		bNeedsToUpdateRayTracingCache = true;
+
+		if (IsRayTracingAllowed() && bSupportRayTracing)
+		{
+			const bool bWantsRayTracingWPO = MaterialRelevance.bUsesWorldPositionOffset;
+
+			if (bWantsRayTracingWPO)
+			{
+				bDynamicRayTracingGeometry = true;
+				bNeedsDynamicRayTracingGeometries = true;
+			}
+		}
+#endif
+
 	}
 
 	virtual ~FCableSceneProxy()
@@ -90,6 +134,14 @@ public:
 		VertexBuffers.ColorVertexBuffer.ReleaseResource();
 		IndexBuffer.ReleaseResource();
 		VertexFactory.ReleaseResource();
+
+#if RHI_RAYTRACING
+		if (IsRayTracingEnabled())
+		{
+			ReleaseStaticRayTracingGeometries();
+			ReleaseDynamicRayTracingGeometries();
+		}
+#endif
 	}
 
 	int32 GetRequiredVertexCount() const
@@ -229,6 +281,14 @@ public:
 			FMemory::Memcpy(IndexBufferData, &Indices[0], Indices.Num() * sizeof(int32));
 			RHICmdList.UnlockBuffer(IndexBuffer.IndexBufferRHI);
 
+#if RHI_RAYTRACING
+			FRayTracingGeometry& RayTracingGeometry = StaticRayTracingGeometry;
+			RayTracingGeometry.ReleaseResource();
+			UpdateRayTracingGeometry_RenderingThread(RayTracingGeometry);
+
+			bNeedsToUpdateRayTracingCache = true;
+#endif
+
 			delete NewDynamicData;
 			NewDynamicData = NULL;
 		}
@@ -321,6 +381,182 @@ public:
 
 	uint32 GetAllocatedSize( void ) const { return( FPrimitiveSceneProxy::GetAllocatedSize() ); }
 
+	virtual void CreateRenderThreadResources()override 
+	{
+#if RHI_RAYTRACING
+
+		if (bSupportRayTracing)
+		{
+			FRayTracingGeometry& RayTracingGeometry = StaticRayTracingGeometry;
+			UpdateRayTracingGeometry_RenderingThread(RayTracingGeometry);
+		}
+
+		if (IsRayTracingAllowed() && bNeedsDynamicRayTracingGeometries)
+		{
+			check(bDynamicRayTracingGeometry);
+			CreateDynamicRayTracingGeometries();
+		}
+		
+#endif
+	}
+
+#if RHI_RAYTRACING
+	virtual void GetDynamicRayTracingInstances(FRayTracingMaterialGatheringContext& Context, TArray<FRayTracingInstance>& OutRayTracingInstances) override
+	{
+		if (CVarRayTracingCableMeshes.GetValueOnRenderThread() == 0)
+		{
+			return;
+		}
+
+		if (!ensureMsgf(IsRayTracingRelevant(),
+			TEXT("GetDynamicRayTracingInstances() is only expected to be called for scene proxies that are compatible with ray tracing. ")
+			TEXT("RT-relevant primitive gathering code in FDeferredShadingSceneRenderer may be wrong.")))
+		{
+			return;
+		}
+
+		bool bEvaluateWPO = bDynamicRayTracingGeometry && CVarRayTracingCableMeshesWPO.GetValueOnRenderThread() == 1;
+		
+		if (bEvaluateWPO && CVarRayTracingCableMeshesWPOCulling.GetValueOnRenderThread() > 0)
+		{
+			FVector ViewCenter = Context.ReferenceView->ViewMatrices.GetViewOrigin();
+			FVector MeshCenter = GetLocalToWorld().GetOrigin();
+			const float CullingRadius = CVarRayTracingCableMeshesWPOCullingRadius.GetValueOnRenderThread();
+			const float BoundingRadius = GetBounds().SphereRadius;
+
+			if (FVector(ViewCenter - MeshCenter).Size() > (CullingRadius + BoundingRadius))
+			{
+				bEvaluateWPO = false;
+			}
+		}
+		
+		if (!bEvaluateWPO)
+		{
+			if (!StaticRayTracingGeometry.IsValid())
+			{
+				return;
+			}
+		}
+
+		FRayTracingGeometry& Geometry = bEvaluateWPO ? DynamicRayTracingGeometry : StaticRayTracingGeometry;
+
+		if (Geometry.Initializer.TotalPrimitiveCount <= 0)
+		{
+			return;
+		}
+
+		FRayTracingInstance& RayTracingInstance = OutRayTracingInstances.AddDefaulted_GetRef();
+
+		const int32 NumRayTracingMaterialEntries = 1;
+
+		if (bNeedsToUpdateRayTracingCache)
+		{
+			CachedRayTracingMaterials.Reset();
+			CachedRayTracingMaterials.Reserve(NumRayTracingMaterialEntries);
+
+			FMeshBatch& MeshBatch = CachedRayTracingMaterials.AddDefaulted_GetRef();
+
+			MeshBatch.VertexFactory = &VertexFactory;
+			MeshBatch.MaterialRenderProxy = Material->GetRenderProxy();
+			MeshBatch.SegmentIndex = 0;
+			MeshBatch.ReverseCulling = IsLocalToWorldDeterminantNegative();
+			MeshBatch.Type = PT_TriangleList;
+			MeshBatch.DepthPriorityGroup = SDPG_World;
+			MeshBatch.bCanApplyViewModeOverrides = false;
+			MeshBatch.CastRayTracedShadow = IsShadowCast(Context.ReferenceView);
+			MeshBatch.DepthPriorityGroup = GetStaticDepthPriorityGroup();
+
+			FMeshBatchElement& BatchElement = MeshBatch.Elements[0];
+			BatchElement.IndexBuffer = &IndexBuffer;
+
+			bool bHasPrecomputedVolumetricLightmap;
+			FMatrix PreviousLocalToWorld;
+			int32 SingleCaptureIndex;
+			bool bOutputVelocity;
+			GetScene().GetPrimitiveUniformShaderParameters_RenderThread(GetPrimitiveSceneInfo(), bHasPrecomputedVolumetricLightmap, PreviousLocalToWorld, SingleCaptureIndex, bOutputVelocity);
+			bOutputVelocity |= AlwaysHasVelocity();
+
+			FDynamicPrimitiveUniformBuffer& DynamicPrimitiveUniformBuffer = Context.RayTracingMeshResourceCollector.AllocateOneFrameResource<FDynamicPrimitiveUniformBuffer>();
+			DynamicPrimitiveUniformBuffer.Set(GetLocalToWorld(), PreviousLocalToWorld, GetBounds(), GetLocalBounds(), ReceivesDecals(), bHasPrecomputedVolumetricLightmap, bOutputVelocity);
+			BatchElement.PrimitiveUniformBufferResource = &DynamicPrimitiveUniformBuffer.UniformBuffer;
+			BatchElement.FirstIndex = 0;
+			BatchElement.NumPrimitives = GetRequiredIndexCount() / 3;
+			BatchElement.MinVertexIndex = 0;
+			BatchElement.MaxVertexIndex = GetRequiredVertexCount();
+
+
+			RayTracingInstance.MaterialsView = MakeArrayView(CachedRayTracingMaterials);
+			CachedRayTracingInstanceMaskAndFlags = Context.BuildInstanceMaskAndFlags(RayTracingInstance, *this);
+			bNeedsToUpdateRayTracingCache = false;
+		}
+		else
+		{
+			RayTracingInstance.MaterialsView = MakeArrayView(CachedRayTracingMaterials);
+		}
+
+		RayTracingInstance.Geometry = &Geometry;
+		const FMatrix& ThisLocalToWorld = GetLocalToWorld();
+		RayTracingInstance.InstanceTransformsView = MakeArrayView(&ThisLocalToWorld, 1);
+
+		if (bEvaluateWPO && VertexFactory.GetType()->SupportsRayTracingDynamicGeometry())
+		{
+			// Use the shared vertex buffer - needs to be updated every frame
+			FRWBuffer* VertexBuffer = nullptr;
+
+			const uint32 VertexCount = VertexBuffers.PositionVertexBuffer.GetNumVertices() + 1;
+
+			Context.DynamicRayTracingGeometriesToUpdate.Add(
+				FRayTracingDynamicGeometryUpdateParams
+				{
+					CachedRayTracingMaterials, // TODO: this copy can be avoided if FRayTracingDynamicGeometryUpdateParams supported array views
+					false,
+					(uint32)VertexCount,
+					uint32((SIZE_T)VertexCount * sizeof(FVector3f)),
+					Geometry.Initializer.TotalPrimitiveCount,
+					&Geometry,
+					VertexBuffer,
+					true
+				}
+			);
+		}
+
+		// Skip computing the mask and flags in the renderer since we are using cached values.
+		RayTracingInstance.bInstanceMaskAndFlagsDirty = false;
+		RayTracingInstance.MaskAndFlags = CachedRayTracingInstanceMaskAndFlags;
+
+		check(CachedRayTracingMaterials.Num() == RayTracingInstance.GetMaterials().Num());
+		checkf(RayTracingInstance.Geometry->Initializer.Segments.Num() == CachedRayTracingMaterials.Num(), TEXT("Segments/Materials mismatch. Number of segments: %d. Number of Materials: %d."),
+			RayTracingInstance.Geometry->Initializer.Segments.Num(),
+			CachedRayTracingMaterials.Num());
+	}
+
+	virtual bool HasRayTracingRepresentation() const override { return bSupportRayTracing; }
+	virtual bool IsRayTracingRelevant() const override { return true; }
+	virtual bool IsRayTracingStaticRelevant() const override { return false; }
+
+	void UpdateRayTracingGeometry_RenderingThread(FRayTracingGeometry& RayTracingGeometry)
+	{
+		FRayTracingGeometryInitializer Initializer;
+		static const FName DebugName("FCableSceneProxy");
+		static int32 DebugNumber = 0;
+		Initializer.DebugName = FDebugName(DebugName, DebugNumber++);
+		Initializer.IndexBuffer = IndexBuffer.IndexBufferRHI;
+		Initializer.TotalPrimitiveCount = IndexBuffer.NumIndices / 3;
+		Initializer.GeometryType = RTGT_Triangles;
+		Initializer.bFastBuild = true;
+		Initializer.bAllowUpdate = false;
+
+		FRayTracingGeometrySegment Segment;
+		Segment.VertexBuffer = VertexBuffers.PositionVertexBuffer.VertexBufferRHI;
+		Segment.NumPrimitives = Initializer.TotalPrimitiveCount;
+		Segment.MaxVertices = VertexBuffers.PositionVertexBuffer.GetNumVertices();
+		Initializer.Segments.Add(Segment);
+
+		RayTracingGeometry.SetInitializer(Initializer);
+		RayTracingGeometry.InitResource();
+	}
+#endif
+
 private:
 
 	UMaterialInterface* Material;
@@ -337,6 +573,44 @@ private:
 	int32 NumSides;
 
 	float TileMaterial;
+
+#if RHI_RAYTRACING
+	void CreateDynamicRayTracingGeometries()
+	{
+		FRayTracingGeometryInitializer& Initializer = DynamicRayTracingGeometry.Initializer;
+		Initializer = StaticRayTracingGeometry.Initializer;
+		for (FRayTracingGeometrySegment& Segment : Initializer.Segments)
+		{
+			Segment.VertexBuffer = nullptr;
+		}
+		Initializer.bAllowUpdate = true;
+		Initializer.bFastBuild = true;
+		Initializer.Type = ERayTracingGeometryInitializerType::Rendering;
+
+		DynamicRayTracingGeometry.InitResource();
+	}
+
+	void ReleaseDynamicRayTracingGeometries()
+	{
+		DynamicRayTracingGeometry.ReleaseResource();
+	}
+
+	void ReleaseStaticRayTracingGeometries()
+	{
+		StaticRayTracingGeometry.ReleaseResource();
+	}
+
+	bool bSupportRayTracing : 1;
+	bool bDynamicRayTracingGeometry : 1;
+	bool bNeedsDynamicRayTracingGeometries : 1;
+	bool bNeedsToUpdateRayTracingCache : 1;
+
+	FRayTracingGeometry StaticRayTracingGeometry;
+	FRayTracingGeometry DynamicRayTracingGeometry;
+	
+	TArray<FMeshBatch> CachedRayTracingMaterials;
+	FRayTracingMaskAndFlags CachedRayTracingInstanceMaskAndFlags;
+#endif
 };
 
 
