Add ray tracing support to FCableSceneProxyCable (with WPO).

diff --git a/Engine/Plugins/Runtime/CableComponent/Source/CableComponent/Private/CableComponent.cpp b/Engine/Plugins/Runtime/CableComponent/Source/CableComponent/Private/CableComponent.cpp
--- a/Engine/Plugins/Runtime/CableComponent/Source/CableComponent/Private/CableComponent.cpp
+++ b/Engine/Plugins/Runtime/CableComponent/Source/CableComponent/Private/CableComponent.cpp
@@ -284,7 +284,7 @@ public:
 #if RHI_RAYTRACING
 			FRayTracingGeometry& RayTracingGeometry = StaticRayTracingGeometry;
 			RayTracingGeometry.ReleaseResource();
-			UpdateRayTracingGeometry_RenderingThread(RayTracingGeometry);
+			UpdateRayTracingGeometry_RenderingThread(RayTracingGeometry, RHICmdList);
 
 			bNeedsToUpdateRayTracingCache = true;
 #endif
@@ -381,20 +381,20 @@ public:
 
 	uint32 GetAllocatedSize( void ) const { return( FPrimitiveSceneProxy::GetAllocatedSize() ); }
 
-	virtual void CreateRenderThreadResources()override 
+	virtual void CreateRenderThreadResources(FRHICommandListBase& RHICmdList) override
 	{
 #if RHI_RAYTRACING
 
 		if (bSupportRayTracing)
 		{
 			FRayTracingGeometry& RayTracingGeometry = StaticRayTracingGeometry;
-			UpdateRayTracingGeometry_RenderingThread(RayTracingGeometry);
+			UpdateRayTracingGeometry_RenderingThread(RayTracingGeometry, RHICmdList);
 		}
 
 		if (IsRayTracingAllowed() && bNeedsDynamicRayTracingGeometries)
 		{
 			check(bDynamicRayTracingGeometry);
-			CreateDynamicRayTracingGeometries();
+			CreateDynamicRayTracingGeometries(RHICmdList);
 		}
 		
 #endif
@@ -534,7 +534,7 @@ public:
 	virtual bool IsRayTracingRelevant() const override { return true; }
 	virtual bool IsRayTracingStaticRelevant() const override { return false; }
 
-	void UpdateRayTracingGeometry_RenderingThread(FRayTracingGeometry& RayTracingGeometry)
+	void UpdateRayTracingGeometry_RenderingThread(FRayTracingGeometry& RayTracingGeometry, FRHICommandListBase& RHICmdList)
 	{
 		FRayTracingGeometryInitializer Initializer;
 		static const FName DebugName("FCableSceneProxy");
@@ -553,7 +553,7 @@ public:
 		Initializer.Segments.Add(Segment);
 
 		RayTracingGeometry.SetInitializer(Initializer);
-		RayTracingGeometry.InitResource();
+		RayTracingGeometry.InitResource(RHICmdList);
 	}
 #endif
 
@@ -575,7 +575,7 @@ private:
 	float TileMaterial;
 
 #if RHI_RAYTRACING
-	void CreateDynamicRayTracingGeometries()
+	void CreateDynamicRayTracingGeometries(FRHICommandListBase& RHICmdList)
 	{
 		FRayTracingGeometryInitializer& Initializer = DynamicRayTracingGeometry.Initializer;
 		Initializer = StaticRayTracingGeometry.Initializer;
@@ -587,7 +587,7 @@ private:
 		Initializer.bFastBuild = true;
 		Initializer.Type = ERayTracingGeometryInitializerType::Rendering;
 
-		DynamicRayTracingGeometry.InitResource();
+		DynamicRayTracingGeometry.InitResource(RHICmdList);
 	}
 
 	void ReleaseDynamicRayTracingGeometries()
