Fix a crash for' FRenderResource was deleted without being released first' in FCableSceneProxy when ray tracing is not enabled in editor.

diff --git a/Engine/Plugins/Runtime/CableComponent/Source/CableComponent/Private/CableComponent.cpp b/Engine/Plugins/Runtime/CableComponent/Source/CableComponent/Private/CableComponent.cpp
--- a/Engine/Plugins/Runtime/CableComponent/Source/CableComponent/Private/CableComponent.cpp
+++ b/Engine/Plugins/Runtime/CableComponent/Source/CableComponent/Private/CableComponent.cpp
@@ -282,11 +282,16 @@ public:
 			RHICmdList.UnlockBuffer(IndexBuffer.IndexBufferRHI);
 
 #if RHI_RAYTRACING
-			FRayTracingGeometry& RayTracingGeometry = StaticRayTracingGeometry;
-			RayTracingGeometry.ReleaseResource();
-			UpdateRayTracingGeometry_RenderingThread(RayTracingGeometry, RHICmdList);
-
-			bNeedsToUpdateRayTracingCache = true;
+			if (bSupportRayTracing && CVarRayTracingCableMeshes.GetValueOnRenderThread() != 0)
+			{
+				FRayTracingGeometry& RayTracingGeometry = StaticRayTracingGeometry;
+				if (RayTracingGeometry.IsValid())
+				{
+					RayTracingGeometry.ReleaseResource();
+					UpdateRayTracingGeometry_RenderingThread(RayTracingGeometry, RHICmdList);
+					bNeedsToUpdateRayTracingCache = true;
+				}
+			}
 #endif
 
 			delete NewDynamicData;
@@ -592,12 +597,18 @@ private:
 
 	void ReleaseDynamicRayTracingGeometries()
 	{
-		DynamicRayTracingGeometry.ReleaseResource();
+		if (DynamicRayTracingGeometry.IsValid())
+		{
+			DynamicRayTracingGeometry.ReleaseResource();
+		}
 	}
 
 	void ReleaseStaticRayTracingGeometries()
 	{
-		StaticRayTracingGeometry.ReleaseResource();
+		if (StaticRayTracingGeometry.IsValid())
+		{
+			StaticRayTracingGeometry.ReleaseResource();
+		}
 	}
 
 	FRayTracingMaskAndFlags CachedRayTracingInstanceMaskAndFlags;
