Moved cable component dynamic update to a render command pipe.

diff --git a/Engine/Plugins/Runtime/CableComponent/Source/CableComponent/Private/CableComponent.cpp b/Engine/Plugins/Runtime/CableComponent/Source/CableComponent/Private/CableComponent.cpp
--- a/Engine/Plugins/Runtime/CableComponent/Source/CableComponent/Private/CableComponent.cpp
+++ b/Engine/Plugins/Runtime/CableComponent/Source/CableComponent/Private/CableComponent.cpp
@@ -14,6 +14,7 @@
 #include "StaticMeshResources.h"
 #include "SceneInterface.h"
 
+DEFINE_RENDER_COMMAND_PIPE(Cable);
 
 static TAutoConsoleVariable<int32> CVarRayTracingCableMeshes(
 	TEXT("r.RayTracing.Geometry.Cable"),
@@ -92,12 +93,12 @@ public:
 		, NumSides(Component->NumSides)
 		, TileMaterial(Component->TileMaterial)
 	{
-		VertexBuffers.InitWithDummyData(&VertexFactory, GetRequiredVertexCount());
+		VertexBuffers.InitWithDummyData(&UE::RenderCommandPipe::Cable, &VertexFactory, GetRequiredVertexCount());
 
 		IndexBuffer.NumIndices = GetRequiredIndexCount();
 
 		// Enqueue initialization of render resource
-		BeginInitResource(&IndexBuffer);
+		BeginInitResource(&IndexBuffer, &UE::RenderCommandPipe::Cable);
 
 		// Grab material
 		Material = Component->GetMaterial(0);
@@ -1043,8 +1044,8 @@ void UCableComponent::SendRenderDynamicData_Concurrent()
 
 		// Enqueue command to send to render thread
 		FCableSceneProxy* CableSceneProxy = (FCableSceneProxy*)SceneProxy;
-		ENQUEUE_RENDER_COMMAND(FSendCableDynamicData)(
-			[CableSceneProxy, DynamicData](FRHICommandListImmediate& RHICmdList)
+		ENQUEUE_RENDER_COMMAND(FSendCableDynamicData)(UE::RenderCommandPipe::Cable,
+			[CableSceneProxy, DynamicData](FRHICommandListBase& RHICmdList)
 		{
 			CableSceneProxy->SetDynamicData_RenderThread(RHICmdList, DynamicData);
 		});
