Fix crash due to unreleased FRayTracingGeometry in FCableSceneProxy

diff --git a/Engine/Plugins/Runtime/CableComponent/Source/CableComponent/Private/CableComponent.cpp b/Engine/Plugins/Runtime/CableComponent/Source/CableComponent/Private/CableComponent.cpp
--- a/Engine/Plugins/Runtime/CableComponent/Source/CableComponent/Private/CableComponent.cpp
+++ b/Engine/Plugins/Runtime/CableComponent/Source/CableComponent/Private/CableComponent.cpp
@@ -137,11 +137,8 @@ public:
 		VertexFactory.ReleaseResource();
 
 #if RHI_RAYTRACING
-		if (IsRayTracingEnabled())
-		{
-			ReleaseStaticRayTracingGeometries();
-			ReleaseDynamicRayTracingGeometries();
-		}
+		StaticRayTracingGeometry.ReleaseResource();
+		DynamicRayTracingGeometry.ReleaseResource();
 #endif
 	}
 
@@ -596,22 +593,6 @@ private:
 		DynamicRayTracingGeometry.InitResource(RHICmdList);
 	}
 
-	void ReleaseDynamicRayTracingGeometries()
-	{
-		if (DynamicRayTracingGeometry.IsValid())
-		{
-			DynamicRayTracingGeometry.ReleaseResource();
-		}
-	}
-
-	void ReleaseStaticRayTracingGeometries()
-	{
-		if (StaticRayTracingGeometry.IsValid())
-		{
-			StaticRayTracingGeometry.ReleaseResource();
-		}
-	}
-
 	FRayTracingMaskAndFlags CachedRayTracingInstanceMaskAndFlags;
 	bool bSupportRayTracing : 1;
 	bool bDynamicRayTracingGeometry : 1;
