Cable Component - static relevance
Cable component proxy does not has any view dependent draw arguments so it can go through DrawStaticElements.
No need to recreate mesh batch every frame.

diff --git a/Engine/Plugins/Runtime/CableComponent/Source/CableComponent/Private/CableComponent.cpp b/Engine/Plugins/Runtime/CableComponent/Source/CableComponent/Private/CableComponent.cpp
--- a/Engine/Plugins/Runtime/CableComponent/Source/CableComponent/Private/CableComponent.cpp
+++ b/Engine/Plugins/Runtime/CableComponent/Source/CableComponent/Private/CableComponent.cpp
@@ -297,6 +297,33 @@ public:
 		}
 	}
 
+	virtual void DrawStaticElements(FStaticPrimitiveDrawInterface* PDI) override
+	{
+		checkSlow(IsInParallelRenderingThread());
+
+		if (!HasViewDependentDPG())
+		{
+			FMeshBatch Mesh;
+			Mesh.VertexFactory = &VertexFactory;
+			Mesh.MaterialRenderProxy = Material->GetRenderProxy();
+			Mesh.ReverseCulling = IsLocalToWorldDeterminantNegative();
+			Mesh.Type = PT_TriangleList;
+			Mesh.DepthPriorityGroup = SDPG_World;
+			Mesh.MeshIdInPrimitive = 0;
+			Mesh.LODIndex = 0;
+			Mesh.SegmentIndex = 0;
+
+			FMeshBatchElement& BatchElement = Mesh.Elements[0];
+			BatchElement.IndexBuffer = &IndexBuffer;
+			BatchElement.FirstIndex = 0;
+			BatchElement.NumPrimitives = GetRequiredIndexCount() / 3;
+			BatchElement.MinVertexIndex = 0;
+			BatchElement.MaxVertexIndex = GetRequiredVertexCount();
+
+			PDI->DrawMesh(Mesh, FLT_MAX);
+		}
+	}
+
 	virtual void GetDynamicMeshElements(const TArray<const FSceneView*>& Views, const FSceneViewFamily& ViewFamily, uint32 VisibilityMap, FMeshElementCollector& Collector) const override
 	{
 		QUICK_SCOPE_CYCLE_COUNTER( STAT_CableSceneProxy_GetDynamicMeshElements );
@@ -371,7 +398,35 @@ public:
 		Result.bRenderInMainPass = ShouldRenderInMainPass();
 		Result.bUsesLightingChannels = GetLightingChannelMask() != GetDefaultLightingChannelMask();
 		Result.bTranslucentSelfShadow = bCastVolumetricTranslucentShadow;
-		Result.bDynamicRelevance = true;
+		const bool bAllowStaticLighting = FReadOnlyCVARCache::Get().bAllowStaticLighting;
+		if (
+#if !(UE_BUILD_SHIPPING) || WITH_EDITOR
+			IsRichView(*View->Family) ||
+			View->Family->EngineShowFlags.Collision ||
+			View->Family->EngineShowFlags.Bounds ||
+			View->Family->EngineShowFlags.VisualizeInstanceUpdates ||
+#endif
+#if WITH_EDITOR
+			(IsSelected() && View->Family->EngineShowFlags.VertexColors) ||
+			(IsSelected() && View->Family->EngineShowFlags.PhysicalMaterialMasks) ||
+#endif
+			// Force down dynamic rendering path if invalid lightmap settings, so we can apply an error material in DrawRichMesh
+			(bAllowStaticLighting && HasStaticLighting() && !HasValidSettingsForStaticLighting()) ||
+			HasViewDependentDPG()
+			)
+		{
+			Result.bDynamicRelevance = true;
+		}
+		else
+		{
+			Result.bStaticRelevance = true;
+
+#if WITH_EDITOR
+			//only check these in the editor
+			Result.bEditorVisualizeLevelInstanceRelevance = IsEditingLevelInstanceChild();
+			Result.bEditorStaticSelectionRelevance = (IsSelected() || IsHovered());
+#endif
+		}
 
 		MaterialRelevance.SetPrimitiveViewRelevance(Result);
 
