PR #10953: Adding GetCustomPrimitiveData() to CableSceneProxy so that custom primitive data can be used

diff --git a/Engine/Plugins/Runtime/CableComponent/Source/CableComponent/Private/CableComponent.cpp b/Engine/Plugins/Runtime/CableComponent/Source/CableComponent/Private/CableComponent.cpp
--- a/Engine/Plugins/Runtime/CableComponent/Source/CableComponent/Private/CableComponent.cpp
+++ b/Engine/Plugins/Runtime/CableComponent/Source/CableComponent/Private/CableComponent.cpp
@@ -368,7 +368,7 @@ public:
 				bOutputVelocity |= AlwaysHasVelocity();
 
 				FDynamicPrimitiveUniformBuffer& DynamicPrimitiveUniformBuffer = Collector.AllocateOneFrameResource<FDynamicPrimitiveUniformBuffer>();
-				DynamicPrimitiveUniformBuffer.Set(Collector.GetRHICommandList(), GetLocalToWorld(), PreviousLocalToWorld, GetBounds(), GetLocalBounds(), true, bHasPrecomputedVolumetricLightmap, bOutputVelocity);
+				DynamicPrimitiveUniformBuffer.Set(Collector.GetRHICommandList(), GetLocalToWorld(), PreviousLocalToWorld, GetBounds(), GetLocalBounds(), GetLocalBounds(), true, bHasPrecomputedVolumetricLightmap, bOutputVelocity, GetCustomPrimitiveData());
 				BatchElement.PrimitiveUniformBufferResource = &DynamicPrimitiveUniformBuffer.UniformBuffer;
 
 				BatchElement.FirstIndex = 0;
@@ -535,7 +535,7 @@ public:
 			bOutputVelocity |= AlwaysHasVelocity();
 
 			FDynamicPrimitiveUniformBuffer& DynamicPrimitiveUniformBuffer = Context.RayTracingMeshResourceCollector.AllocateOneFrameResource<FDynamicPrimitiveUniformBuffer>();
-			DynamicPrimitiveUniformBuffer.Set(Context.RHICmdList, GetLocalToWorld(), PreviousLocalToWorld, GetBounds(), GetLocalBounds(), ReceivesDecals(), bHasPrecomputedVolumetricLightmap, bOutputVelocity);
+			DynamicPrimitiveUniformBuffer.Set(Context.RHICmdList, GetLocalToWorld(), PreviousLocalToWorld, GetBounds(), GetLocalBounds(), GetLocalBounds(), ReceivesDecals(), bHasPrecomputedVolumetricLightmap, bOutputVelocity, GetCustomPrimitiveData());
 			BatchElement.PrimitiveUniformBufferResource = &DynamicPrimitiveUniformBuffer.UniformBuffer;
 			BatchElement.FirstIndex = 0;
 			BatchElement.NumPrimitives = GetRequiredIndexCount() / 3;
