Fix incorrect ray tracing WPO culling.
- use bounds center instead of primitive origin

diff --git a/Engine/Plugins/Runtime/CableComponent/Source/CableComponent/Private/CableComponent.cpp b/Engine/Plugins/Runtime/CableComponent/Source/CableComponent/Private/CableComponent.cpp
--- a/Engine/Plugins/Runtime/CableComponent/Source/CableComponent/Private/CableComponent.cpp
+++ b/Engine/Plugins/Runtime/CableComponent/Source/CableComponent/Private/CableComponent.cpp
@@ -476,8 +476,8 @@ public:
 		
 		if (bEvaluateWPO && CVarRayTracingCableMeshesWPOCulling.GetValueOnRenderThread() > 0)
 		{
-			FVector ViewCenter = Context.ReferenceView->ViewMatrices.GetViewOrigin();
-			FVector MeshCenter = GetLocalToWorld().GetOrigin();
+			const FVector ViewCenter = Context.ReferenceView->ViewMatrices.GetViewOrigin();
+			const FVector MeshCenter = GetBounds().Origin;
 			const float CullingRadius = CVarRayTracingCableMeshesWPOCullingRadius.GetValueOnRenderThread();
 			const float BoundingRadius = GetBounds().SphereRadius;
 
diff --git a/Engine/Source/Runtime/Engine/Private/StaticMeshRender.cpp b/Engine/Source/Runtime/Engine/Private/StaticMeshRender.cpp
--- a/Engine/Source/Runtime/Engine/Private/StaticMeshRender.cpp
+++ b/Engine/Source/Runtime/Engine/Private/StaticMeshRender.cpp
@@ -1906,8 +1906,8 @@ void FStaticMeshSceneProxy::GetDynamicRayTracingInstances(FRayTracingMaterialGat
 
 	if (bEvaluateWPO && CVarRayTracingStaticMeshesWPOCulling.GetValueOnRenderThread() > 0)
 	{
-		FVector ViewCenter = Context.ReferenceView->ViewMatrices.GetViewOrigin();
-		FVector MeshCenter = GetLocalToWorld().GetOrigin();
+		const FVector ViewCenter = Context.ReferenceView->ViewMatrices.GetViewOrigin();
+		const FVector MeshCenter = GetBounds().Origin;
 		const float CullingRadius = CVarRayTracingStaticMeshesWPOCullingRadius.GetValueOnRenderThread();
 		const float BoundingRadius = GetBounds().SphereRadius;
 
