Improvements to dynamic ray tracing geometry managment in some proxy types.
- fixes FRayTracingGeometry not being properly released by FImagePlateSceneProxy and FProceduralMeshSceneProxy when ray tracing is toggled off at runtime.
- small refactor in FCableSceneProxy and FSkeletalMeshSceneProxy

diff --git a/Engine/Plugins/Experimental/ImagePlate/Source/ImagePlate/Private/ImagePlateComponent.cpp b/Engine/Plugins/Experimental/ImagePlate/Source/ImagePlate/Private/ImagePlateComponent.cpp
--- a/Engine/Plugins/Experimental/ImagePlate/Source/ImagePlate/Private/ImagePlateComponent.cpp
+++ b/Engine/Plugins/Experimental/ImagePlate/Source/ImagePlate/Private/ImagePlateComponent.cpp
@@ -80,10 +80,7 @@ namespace
 			IndexBuffer.ReleaseResource();
 			VertexFactory.ReleaseResource();
 #if RHI_RAYTRACING
-			if (IsRayTracingEnabled())
-			{
-				RayTracingGeometry.ReleaseResource();
-			}
+			RayTracingGeometry.ReleaseResource();
 #endif
 		}
 
diff --git a/Engine/Plugins/Runtime/CableComponent/Source/CableComponent/Private/CableComponent.cpp b/Engine/Plugins/Runtime/CableComponent/Source/CableComponent/Private/CableComponent.cpp
--- a/Engine/Plugins/Runtime/CableComponent/Source/CableComponent/Private/CableComponent.cpp
+++ b/Engine/Plugins/Runtime/CableComponent/Source/CableComponent/Private/CableComponent.cpp
@@ -107,21 +107,7 @@ public:
 
 #if RHI_RAYTRACING
 		bSupportRayTracing = IsRayTracingEnabled();
-		bDynamicRayTracingGeometry = false;
-		bNeedsDynamicRayTracingGeometries = false;
-		
-		bNeedsToUpdateRayTracingCache = true;
-
-		if (IsRayTracingAllowed() && bSupportRayTracing)
-		{
-			const bool bWantsRayTracingWPO = MaterialRelevance.bUsesWorldPositionOffset;
-
-			if (bWantsRayTracingWPO)
-			{
-				bDynamicRayTracingGeometry = true;
-				bNeedsDynamicRayTracingGeometries = true;
-			}
-		}
+		bDynamicRayTracingGeometry = bSupportRayTracing && MaterialRelevance.bUsesWorldPositionOffset;
 #endif
 
 		ENQUEUE_RENDER_COMMAND(InitCableResources)(UE::RenderCommandPipe::Cable,
@@ -134,13 +120,13 @@ public:
 			if (bSupportRayTracing)
 			{
 				FRayTracingGeometry& RayTracingGeometry = StaticRayTracingGeometry;
-				UpdateRayTracingGeometry_RenderingThread(RayTracingGeometry, RHICmdList);
-			}
+				CreateRayTracingGeometry_RenderingThread(RayTracingGeometry, RHICmdList);
+				bNeedsToUpdateRayTracingCache = true;
 
-			if (IsRayTracingAllowed() && bNeedsDynamicRayTracingGeometries)
-			{
-				check(bDynamicRayTracingGeometry);
-				CreateDynamicRayTracingGeometries(RHICmdList);
+				if (bDynamicRayTracingGeometry)
+				{
+					CreateDynamicRayTracingGeometries(RHICmdList);
+				}
 			}
 #endif
 		});
@@ -304,7 +290,7 @@ public:
 				if (RayTracingGeometry.IsValid())
 				{
 					RayTracingGeometry.ReleaseResource();
-					UpdateRayTracingGeometry_RenderingThread(RayTracingGeometry, RHICmdList);
+					CreateRayTracingGeometry_RenderingThread(RayTracingGeometry, RHICmdList);
 					bNeedsToUpdateRayTracingCache = true;
 				}
 			}
@@ -557,6 +543,7 @@ public:
 		const FMatrix& ThisLocalToWorld = GetLocalToWorld();
 		RayTracingInstance.InstanceTransformsView = MakeArrayView(&ThisLocalToWorld, 1);
 
+		// TODO: Checking if VertexFactory.GetType()->SupportsRayTracingDynamicGeometry() should be done when initializing bDynamicRayTracingGeometry otherwise we end up with unbuilt BLAS
 		if (bEvaluateWPO && VertexFactory.GetType()->SupportsRayTracingDynamicGeometry())
 		{
 			// Use the shared vertex buffer - needs to be updated every frame
@@ -589,7 +576,7 @@ public:
 	virtual bool IsRayTracingRelevant() const override { return true; }
 	virtual bool IsRayTracingStaticRelevant() const override { return false; }
 
-	void UpdateRayTracingGeometry_RenderingThread(FRayTracingGeometry& RayTracingGeometry, FRHICommandListBase& RHICmdList)
+	void CreateRayTracingGeometry_RenderingThread(FRayTracingGeometry& RayTracingGeometry, FRHICommandListBase& RHICmdList)
 	{
 		FRayTracingGeometryInitializer Initializer;
 		static const FName DebugName("FCableSceneProxy");
@@ -647,7 +634,6 @@ private:
 
 	bool bSupportRayTracing : 1;
 	bool bDynamicRayTracingGeometry : 1;
-	bool bNeedsDynamicRayTracingGeometries : 1;
 	bool bNeedsToUpdateRayTracingCache : 1;
 
 	FRayTracingGeometry StaticRayTracingGeometry;
diff --git a/Engine/Plugins/Runtime/ProceduralMeshComponent/Source/ProceduralMeshComponent/Private/ProceduralMeshComponent.cpp b/Engine/Plugins/Runtime/ProceduralMeshComponent/Source/ProceduralMeshComponent/Private/ProceduralMeshComponent.cpp
--- a/Engine/Plugins/Runtime/ProceduralMeshComponent/Source/ProceduralMeshComponent/Private/ProceduralMeshComponent.cpp
+++ b/Engine/Plugins/Runtime/ProceduralMeshComponent/Source/ProceduralMeshComponent/Private/ProceduralMeshComponent.cpp
@@ -216,10 +216,7 @@ public:
 				Section->VertexFactory.ReleaseResource();
 
 #if RHI_RAYTRACING
-				if (IsRayTracingEnabled())
-				{
-					Section->RayTracingGeometry.ReleaseResource();
-				}
+				Section->RayTracingGeometry.ReleaseResource();
 #endif
 
 				delete Section;
diff --git a/Engine/Source/Runtime/Engine/Private/SkeletalMesh.cpp b/Engine/Source/Runtime/Engine/Private/SkeletalMesh.cpp
--- a/Engine/Source/Runtime/Engine/Private/SkeletalMesh.cpp
+++ b/Engine/Source/Runtime/Engine/Private/SkeletalMesh.cpp
@@ -6949,15 +6949,17 @@ void FSkeletalMeshSceneProxy::GetDynamicRayTracingInstances(FRayTracingMaterialG
 		return;
 	}
 
+	FRayTracingGeometry* RayTracingGeometry = MeshObject->GetRayTracingGeometry();
+
 	// GetRayTracingGeometry()->IsInitialized() is checked as a workaround for UE-92634. FSkeletalMeshSceneProxy's resources may have already been released, but proxy has not removed yet)
-	if (MeshObject->GetRayTracingGeometry() && MeshObject->GetRayTracingGeometry()->IsInitialized())
+	if (RayTracingGeometry && RayTracingGeometry->IsInitialized())
 	{
-		if(MeshObject->GetRayTracingGeometry()->IsValid())
+		if(RayTracingGeometry->IsValid())
 		{
-			check(MeshObject->GetRayTracingGeometry()->Initializer.IndexBuffer.IsValid());
+			check(RayTracingGeometry->Initializer.IndexBuffer.IsValid());
 			
 			FRayTracingInstance RayTracingInstance;
-			RayTracingInstance.Geometry = MeshObject->GetRayTracingGeometry();
+			RayTracingInstance.Geometry = RayTracingGeometry;
 
 			// Setup materials for each segment
 			const int32 LODIndex = MeshObject->GetRayTracingLOD();
@@ -6975,7 +6977,7 @@ void FSkeletalMeshSceneProxy::GetDynamicRayTracingInstances(FRayTracingMaterialG
 			check(LODSection.SectionElements.Num() == LODData.RenderSections.Num());
 				
 			//#dxr_todo (UE-113617): verify why this condition is not fulfilled sometimes
-			if(!ensure(LODSection.SectionElements.Num() == MeshObject->GetRayTracingGeometry()->Initializer.Segments.Num()))
+			if(!ensure(LODSection.SectionElements.Num() == RayTracingGeometry->Initializer.Segments.Num()))
 			{
 				return;
 			}
@@ -7030,7 +7032,7 @@ void FSkeletalMeshSceneProxy::GetDynamicRayTracingInstances(FRayTracingMaterialG
 					GeometrySections.Add(Segment);
 				}
 
-				MeshObject->GetRayTracingGeometry()->Initializer.Segments = GeometrySections;
+				RayTracingGeometry->Initializer.Segments = GeometrySections;
 
 				Context.DynamicRayTracingGeometriesToUpdate.Add(
 					FRayTracingDynamicGeometryUpdateParams
@@ -7039,8 +7041,8 @@ void FSkeletalMeshSceneProxy::GetDynamicRayTracingInstances(FRayTracingMaterialG
 						false,
 						LODData.GetNumVertices(),
 						LODData.GetNumVertices() * (uint32)sizeof(FVector3f),
-						MeshObject->GetRayTracingGeometry()->Initializer.TotalPrimitiveCount,
-						MeshObject->GetRayTracingGeometry(),
+						RayTracingGeometry->Initializer.TotalPrimitiveCount,
+						RayTracingGeometry,
 						MeshObject->GetRayTracingDynamicVertexBuffer(),
 						true
 					}
