Add teleport corrections to CableComponent.

Currently, large movement of the cable component causes instability as cable end points can be moved without moving the rest of the cable particles with them. This change adds the option for users to enable teleports similar to other systems like cloth. These teleports can either retransform the cable particles to reflect the new cable position or reset the particle positions entirely, avoiding the instability caused by just teleporting cable end points.

diff --git a/Engine/Plugins/Runtime/CableComponent/Source/CableComponent/Classes/CableComponent.h b/Engine/Plugins/Runtime/CableComponent/Source/CableComponent/Classes/CableComponent.h
--- a/Engine/Plugins/Runtime/CableComponent/Source/CableComponent/Classes/CableComponent.h
+++ b/Engine/Plugins/Runtime/CableComponent/Source/CableComponent/Classes/CableComponent.h
@@ -166,6 +166,33 @@ public:
 	UPROPERTY(EditAnywhere, BlueprintReadWrite, Category="Cable Rendering", meta=(UIMin = "0.1", UIMax = "8"))
 	float TileMaterial;
 
+	/**
+	* Always reset cable particle positions and velocities on a teleport.
+	*/
+	UPROPERTY(EditAnywhere, BlueprintReadWrite, Category = "Teleport")
+	bool bResetAfterTeleport = false;
+
+	/**
+	* Conduct teleportation if the movement of either fixed end point of the cable is greater than this threshold in 1 frame.
+	* Zero or negative values will skip the check.
+	*/
+	UPROPERTY(EditAnywhere, BlueprintReadWrite, Category="Teleport", meta=(UIMin = "0"))
+	float TeleportDistanceThreshold = 500;
+
+	/**
+	* Rotation threshold in degrees.
+	* Conduct teleportation if the cable's rotation is greater than this threshold in 1 frame.
+	* Zero or negative values will skip the check.
+	*/
+	UPROPERTY(EditAnywhere, BlueprintReadWrite, Category="Teleport", meta = (UIMin = "0"))
+	float TeleportRotationThreshold = 0;
+
+	/**
+	* Teleport the cable particles on reattaching any end of the cable.
+	**/
+	UPROPERTY(EditAnywhere, BlueprintReadWrite, Category = "Teleport")
+	bool bTeleportAfterReattach = false;
+
 private:
 
 	/** Solve the cable spring constraints */
@@ -183,6 +210,17 @@ private:
 	/** Array of cable particles */
 	TArray<FCableParticle>	Particles;
 
+	/** Previous cable state for teleport corrections */
+	FTransform LastTransform;
+	FTransform LastEndPointTransform;
+	FVector LastEndPoint;
+	FVector LastEndLocation;
+	FVector LastStartPoint;
+	bool bLastStartAttached;
+	bool bLastEndAttached;
+	
+	/** Perform checks and corrections on particle positions for teleports */
+	void DoTeleportCorrections(const FVector& StartPosition, const FVector& EndPosition);
 
 	friend class FCableSceneProxy;
 };
diff --git a/Engine/Plugins/Runtime/CableComponent/Source/CableComponent/Private/CableComponent.cpp b/Engine/Plugins/Runtime/CableComponent/Source/CableComponent/Private/CableComponent.cpp
--- a/Engine/Plugins/Runtime/CableComponent/Source/CableComponent/Private/CableComponent.cpp
+++ b/Engine/Plugins/Runtime/CableComponent/Source/CableComponent/Private/CableComponent.cpp
@@ -662,6 +662,13 @@ UCableComponent::UCableComponent( const FObjectInitializer& ObjectInitializer )
 	CollisionFriction = 0.2f;
 	CableGravityScale = 1.f;
 
+	GetEndPositions(LastStartPoint,LastEndPoint);
+	LastTransform = GetComponentTransform();
+	LastEndPointTransform = GetComponentTransform();
+	LastEndLocation = EndLocation;
+	bLastStartAttached = bAttachStart;
+	bLastEndAttached = bAttachEnd;
+
 	SetCollisionProfileName(UCollisionProfile::PhysicsActor_ProfileName);
 }
 
@@ -897,6 +904,153 @@ void UCableComponent::GetCableParticleLocations(TArray<FVector>& Locations) cons
 	}
 }
 
+void UCableComponent::DoTeleportCorrections(const FVector& StartPosition, const FVector& EndPosition)
+{
+	const FTransform ComponentTransform = GetComponentTransform();
+	USceneComponent* EndComponent = Cast<USceneComponent>(AttachEndTo.GetComponent(GetOwner()));
+	if (EndComponent == NULL)
+	{
+		EndComponent = this;
+	}
+
+	FTransform EndPointTransform = FTransform::Identity;
+
+	if (AttachEndToSocketName != NAME_None)
+	{
+		EndPointTransform = EndComponent->GetSocketTransform(AttachEndToSocketName);
+	}
+	else
+	{
+		EndPointTransform = EndComponent->GetComponentTransform();
+	}
+
+	bool bShouldCorrect = false;
+	const bool bReattached = (bAttachStart && !bLastStartAttached) || (bAttachEnd && !bLastEndAttached);
+	const bool bAttachmentChange = (bAttachStart != bLastStartAttached) || (bAttachEnd != bLastEndAttached);
+
+	//Cannot apply corrections unless some end of the cable is not free falling.
+	if (bAttachStart || bAttachEnd)
+	{
+		if (bReattached && bTeleportAfterReattach)
+		{
+			bShouldCorrect = true;
+		}
+		else
+		{
+			const float SquaredDistanceThreshold = TeleportDistanceThreshold * TeleportDistanceThreshold;
+
+			const float StartDifference = FVector::DistSquared(StartPosition, LastStartPoint);
+			const float EndDifference = FVector::DistSquared(EndPosition, LastEndPoint);
+
+			if (SquaredDistanceThreshold > 0 && (StartDifference > SquaredDistanceThreshold || EndDifference > SquaredDistanceThreshold))
+			{
+				bShouldCorrect = true;
+			}
+			else if(!bAttachmentChange && !(bAttachStart && bAttachEnd && EndComponent != this))
+			{
+				//Only calculate rotational difference if we have two consecutive frames with transform data. 
+				//Don't attempt to calculate rotation if the end points can be controlled by separate transforms.
+
+				const FTransform CurrentRootTransform = bAttachEnd ?  EndPointTransform : ComponentTransform;
+				const FTransform LastRootTransform = bAttachEnd ? LastEndPointTransform : LastTransform;
+
+				const float RotationDifference = CurrentRootTransform.GetRotation().AngularDistance(LastRootTransform.GetRotation());
+
+				if (TeleportRotationThreshold > 0 && FMath::RadiansToDegrees(FMath::Abs(RotationDifference)) > TeleportRotationThreshold)
+				{
+					bShouldCorrect = true;
+				}
+			}
+		}
+	}
+
+	if (bShouldCorrect)
+	{
+		bool bCorrectiveTransform = false;
+		bool bCorrectiveReinitialize = false;
+		bool bAdjustStart = false;
+		bool bAdjustEnd = false;
+
+		FTransform DeltaTransform = FTransform::Identity;
+		const FVector DeltaEndLocation = EndLocation - LastEndLocation;
+
+		if (bAttachmentChange || bResetAfterTeleport)
+		{
+			bCorrectiveReinitialize = true;
+		}
+		else
+		{
+			if (bAttachStart && bAttachEnd)
+			{
+				if (DeltaEndLocation.SquaredLength() > 0 || EndComponent != this)
+				{
+					//If the end points can change independently of the starts we conservatively reinitialize
+					bCorrectiveReinitialize = true;
+				}
+				else
+				{
+					//Start and end points will both move using the cable transform. 
+					DeltaTransform = (LastTransform.Inverse() * ComponentTransform);
+					bCorrectiveTransform = true;
+				}
+			}
+			else
+			{
+				bCorrectiveTransform = true;
+				if (!bAttachStart)
+				{
+					//Detached start: Use cable transform to correct if the cable isn't attached to another object, otherwise use the transform of that object.
+					const FTransform CurrentRootTransform = EndComponent == this ? ComponentTransform : EndPointTransform;
+					const FTransform LastRootTransform = EndComponent == this ? LastTransform : LastEndPointTransform;
+
+					DeltaTransform = FTransform(CurrentRootTransform.TransformVector(DeltaEndLocation)) * (LastRootTransform.Inverse() * CurrentRootTransform);
+					bAdjustStart = true;
+				}
+				else if (!bAttachEnd)
+				{
+					//Detached end: Use cable transform to correct.
+					DeltaTransform = (LastTransform.Inverse() * ComponentTransform);
+					bAdjustEnd = true;
+				}
+			}
+		}
+
+		if (bCorrectiveTransform)
+		{
+			int Index = bAdjustStart ? 0 : 1;
+			const int LastIndex = bAdjustEnd ? NumSegments : NumSegments - 1;
+
+			for (; Index <= LastIndex; Index++)
+			{
+				Particles[Index].Position = DeltaTransform.TransformPosition(Particles[Index].Position);
+				Particles[Index].OldPosition = DeltaTransform.TransformPosition(Particles[Index].OldPosition);
+			}
+		}
+		else if (bCorrectiveReinitialize)
+		{
+			const FVector Delta = EndPosition - StartPosition;
+			for (int32 ParticleIdx = 0; ParticleIdx <= NumSegments; ParticleIdx++)
+			{
+				FCableParticle& Particle = Particles[ParticleIdx];
+
+				const float Alpha = (float)ParticleIdx / (float)NumSegments;
+				const FVector InitialPosition = StartPosition + (Alpha * Delta);
+
+				Particle.Position = InitialPosition;
+				Particle.OldPosition = InitialPosition;
+			}
+		}
+	}
+
+	LastTransform = ComponentTransform;
+	LastEndPointTransform = EndPointTransform;
+	LastEndLocation = EndLocation;
+	LastStartPoint = StartPosition;
+	LastEndPoint = EndPosition;
+	bLastStartAttached = bAttachStart;
+	bLastEndAttached = bAttachEnd;
+}
+
 
 void UCableComponent::GetEndPositions(FVector& OutStartPosition, FVector& OutEndPosition)
 {
@@ -977,6 +1131,8 @@ void UCableComponent::TickComponent(float DeltaTime, enum ELevelTick TickType, F
 		EndParticle.bFree = true;
 	}
 
+	DoTeleportCorrections(CableStart, CableEnd);
+
 	// Ensure a non-zero substep
 	float UseSubstep = FMath::Max(SubstepTime, 0.005f);
 
