CableComponent doesn't set previous state in OnRegister.

This change sets the previous cable state in OnRegister to avoid the cable making an incorrect teleport transform on initialization.

diff --git a/Engine/Plugins/Runtime/CableComponent/Source/CableComponent/Private/CableComponent.cpp b/Engine/Plugins/Runtime/CableComponent/Source/CableComponent/Private/CableComponent.cpp
--- a/Engine/Plugins/Runtime/CableComponent/Source/CableComponent/Private/CableComponent.cpp
+++ b/Engine/Plugins/Runtime/CableComponent/Source/CableComponent/Private/CableComponent.cpp
@@ -707,6 +707,27 @@ void UCableComponent::OnRegister()
 		Particle.OldPosition = InitialPosition;
 		Particle.bFree = true; // default to free, will be fixed if desired in TickComponent
 	}
+
+	LastStartPoint = CableStart;
+	LastEndPoint = CableEnd;
+	LastEndLocation = EndLocation;
+	LastTransform = GetComponentTransform();
+	USceneComponent* EndComponent = Cast<USceneComponent>(AttachEndTo.GetComponent(GetOwner()));
+	if (EndComponent == NULL)
+	{
+		EndComponent = this;
+	}
+
+	if (AttachEndToSocketName != NAME_None)
+	{
+		LastEndPointTransform = EndComponent->GetSocketTransform(AttachEndToSocketName);
+	}
+	else
+	{
+		LastEndPointTransform = EndComponent->GetComponentTransform();
+	}
+	bLastStartAttached = bAttachStart;
+	bLastEndAttached = bAttachEnd;
 }
 
 void UCableComponent::VerletIntegrate(float InSubstepTime, const FVector& Gravity)
