Avoid temporary array creation for cable mesh building.

The existing cable mesh build first copies data to a temporary array then copies it to the RHI buffer. This temporary array is reallocated every time the function is called, causing performance issues with larger cables. This change copies directly to the RHI buffers and bypasses making the temporary array.

diff --git a/Engine/Plugins/Runtime/CableComponent/Source/CableComponent/Private/CableComponent.cpp b/Engine/Plugins/Runtime/CableComponent/Source/CableComponent/Private/CableComponent.cpp
--- a/Engine/Plugins/Runtime/CableComponent/Source/CableComponent/Private/CableComponent.cpp
+++ b/Engine/Plugins/Runtime/CableComponent/Source/CableComponent/Private/CableComponent.cpp
@@ -167,25 +167,28 @@ public:
 		return (AlongIdx * (NumSides+1)) + AroundIdx;
 	}
 
-	void BuildCableMesh(const TArray<FVector>& InPoints, TArray<FDynamicMeshVertex>& OutVertices, TArray<int32>& OutIndices)
+	void BuildCableMesh(FRHICommandListBase& RHICmdList, const FCableDynamicData* NewDynamicData)
 	{
-		const FColor VertexColor(255,255,255);
+		const TArray<FVector>& InPoints = NewDynamicData->CablePoints;
+		const FColor VertexColor(255, 255, 255);
 		const int32 NumPoints = InPoints.Num();
-		const int32 SegmentCount = NumPoints-1;
+		const int32 SegmentCount = NumPoints - 1;
 
 		// Build vertices
 
 		// We double up the first and last vert of the ring, because the UVs are different
-		int32 NumRingVerts = NumSides+1;
+		int32 NumRingVerts = NumSides + 1;
+
+		uint32 VertIdx = 0;
 
 		// For each point along spline..
-		for(int32 PointIdx=0; PointIdx<NumPoints; PointIdx++)
+		for (int32 PointIdx = 0; PointIdx < NumPoints; PointIdx++)
 		{
-			const float AlongFrac = (float)PointIdx/(float)SegmentCount; // Distance along cable
+			const float AlongFrac = (float)PointIdx / (float)SegmentCount; // Distance along cable
 
 			// Find direction of cable at this point, by averaging previous and next points
-			const int32 PrevIndex = FMath::Max(0, PointIdx-1);
-			const int32 NextIndex = FMath::Min(PointIdx+1, NumPoints-1);
+			const int32 PrevIndex = FMath::Max(0, PointIdx - 1);
+			const int32 NextIndex = FMath::Min(PointIdx + 1, NumPoints - 1);
 			const FVector ForwardDir = (InPoints[NextIndex] - InPoints[PrevIndex]).GetSafeNormal();
 
 			// Find quat from up (Z) vector to forward
@@ -196,42 +199,78 @@ public:
 			const FVector UpDir = DeltaQuat.RotateVector(FVector(1, 0, 0));
 
 			// Generate a ring of verts
-			for(int32 VertIdx = 0; VertIdx<NumRingVerts; VertIdx++)
+			for (int32 RingVertIdx = 0; RingVertIdx < NumRingVerts; RingVertIdx++)
 			{
-				const float AroundFrac = float(VertIdx)/float(NumSides);
+				const float AroundFrac = float(RingVertIdx) / float(NumSides);
 				// Find angle around the ring
 				const float RadAngle = 2.f * PI * AroundFrac;
 				// Find direction from center of cable to this vertex
 				const FVector OutDir = (FMath::Cos(RadAngle) * UpDir) + (FMath::Sin(RadAngle) * RightDir);
 
-				FDynamicMeshVertex Vert;
-				Vert.Position = FVector3f(InPoints[PointIdx] + (OutDir * 0.5f * CableWidth));
-				Vert.TextureCoordinate[0] = FVector2f(AlongFrac * TileMaterial, AroundFrac);
-				Vert.Color = VertexColor;
-				Vert.SetTangents((FVector3f)ForwardDir, FVector3f(OutDir ^ ForwardDir), (FVector3f)OutDir);
-				OutVertices.Add(Vert);
+				VertexBuffers.PositionVertexBuffer.VertexPosition(VertIdx) = FVector3f(InPoints[PointIdx] + (OutDir * 0.5f * CableWidth));
+				VertexBuffers.StaticMeshVertexBuffer.SetVertexUV(VertIdx, 0, FVector2f(AlongFrac * TileMaterial, AroundFrac));
+				VertexBuffers.ColorVertexBuffer.VertexColor(VertIdx) = VertexColor;
+				VertexBuffers.StaticMeshVertexBuffer.SetVertexTangents(VertIdx, (FVector3f)ForwardDir, FVector3f(OutDir ^ ForwardDir), (FVector3f)OutDir);
+				VertIdx++;
 			}
 		}
 
+		check(VertIdx == GetRequiredVertexCount());
+
+		{
+			FPositionVertexBuffer& VertexBuffer = VertexBuffers.PositionVertexBuffer;
+			void* VertexBufferData = RHICmdList.LockBuffer(VertexBuffer.VertexBufferRHI, 0, VertexBuffer.GetNumVertices() * VertexBuffer.GetStride(), RLM_WriteOnly);
+			FMemory::Memcpy(VertexBufferData, VertexBuffer.GetVertexData(), VertexBuffer.GetNumVertices() * VertexBuffer.GetStride());
+			RHICmdList.UnlockBuffer(VertexBuffer.VertexBufferRHI);
+		}
+
+		{
+			FColorVertexBuffer& ColorBuffer = VertexBuffers.ColorVertexBuffer;
+			void* VertexBufferData = RHICmdList.LockBuffer(ColorBuffer.VertexBufferRHI, 0, ColorBuffer.GetNumVertices() * ColorBuffer.GetStride(), RLM_WriteOnly);
+			FMemory::Memcpy(VertexBufferData, ColorBuffer.GetVertexData(), ColorBuffer.GetNumVertices() * ColorBuffer.GetStride());
+			RHICmdList.UnlockBuffer(ColorBuffer.VertexBufferRHI);
+		}
+
+		{
+			FStaticMeshVertexBuffer& VertexBuffer = VertexBuffers.StaticMeshVertexBuffer;
+			void* VertexBufferData = RHICmdList.LockBuffer(VertexBuffer.TangentsVertexBuffer.VertexBufferRHI, 0, VertexBuffer.GetTangentSize(), RLM_WriteOnly);
+			FMemory::Memcpy(VertexBufferData, VertexBuffer.GetTangentData(), VertexBuffer.GetTangentSize());
+			RHICmdList.UnlockBuffer(VertexBuffer.TangentsVertexBuffer.VertexBufferRHI);
+		}
+
+		{
+			FStaticMeshVertexBuffer& VertexBuffer = VertexBuffers.StaticMeshVertexBuffer;
+			void* VertexBufferData = RHICmdList.LockBuffer(VertexBuffer.TexCoordVertexBuffer.VertexBufferRHI, 0, VertexBuffer.GetTexCoordSize(), RLM_WriteOnly);
+			FMemory::Memcpy(VertexBufferData, VertexBuffer.GetTexCoordData(), VertexBuffer.GetTexCoordSize());
+			RHICmdList.UnlockBuffer(VertexBuffer.TexCoordVertexBuffer.VertexBufferRHI);
+		}
+
+		int32* IndexBufferData = static_cast<int32*>(RHICmdList.LockBuffer(IndexBuffer.IndexBufferRHI, 0, GetRequiredIndexCount() * sizeof(int32), RLM_WriteOnly));
+		uint32 IndexBufferIndex = 0;
+
 		// Build triangles
-		for(int32 SegIdx=0; SegIdx<SegmentCount; SegIdx++)
+		for (int32 SegIdx = 0; SegIdx < SegmentCount; SegIdx++)
 		{
-			for(int32 SideIdx=0; SideIdx<NumSides; SideIdx++)
+			for (int32 SideIdx = 0; SideIdx < NumSides; SideIdx++)
 			{
 				int32 TL = GetVertIndex(SegIdx, SideIdx);
-				int32 BL = GetVertIndex(SegIdx, SideIdx+1);
-				int32 TR = GetVertIndex(SegIdx+1, SideIdx);
-				int32 BR = GetVertIndex(SegIdx+1, SideIdx+1);
+				int32 BL = GetVertIndex(SegIdx, SideIdx + 1);
+				int32 TR = GetVertIndex(SegIdx + 1, SideIdx);
+				int32 BR = GetVertIndex(SegIdx + 1, SideIdx + 1);
 
-				OutIndices.Add(TL);
-				OutIndices.Add(BL);
-				OutIndices.Add(TR);
+				IndexBufferData[IndexBufferIndex++] = TL;
+				IndexBufferData[IndexBufferIndex++] = BL;
+				IndexBufferData[IndexBufferIndex++] = TR;
 
-				OutIndices.Add(TR);
-				OutIndices.Add(BL);
-				OutIndices.Add(BR);
+				IndexBufferData[IndexBufferIndex++] = TR;
+				IndexBufferData[IndexBufferIndex++] = BL;
+				IndexBufferData[IndexBufferIndex++] = BR;
 			}
 		}
+
+		RHICmdList.UnlockBuffer(IndexBuffer.IndexBufferRHI);
+
+		check(IndexBufferIndex == GetRequiredIndexCount());
 	}
 
 	/** Called on render thread to assign new dynamic data */
@@ -240,54 +279,7 @@ public:
 		if (NewDynamicData != nullptr)
 		{
 			// Build mesh from cable points
-			TArray<FDynamicMeshVertex> Vertices;
-			TArray<int32> Indices;
-			BuildCableMesh(NewDynamicData->CablePoints, Vertices, Indices);
-
-			check(Vertices.Num() == GetRequiredVertexCount());
-			check(Indices.Num() == GetRequiredIndexCount());
-
-			for (int i = 0; i < Vertices.Num(); i++)
-			{
-				const FDynamicMeshVertex& Vertex = Vertices[i];
-
-				VertexBuffers.PositionVertexBuffer.VertexPosition(i) = Vertex.Position;
-				VertexBuffers.StaticMeshVertexBuffer.SetVertexTangents(i, Vertex.TangentX.ToFVector3f(), Vertex.GetTangentY(), Vertex.TangentZ.ToFVector3f());
-				VertexBuffers.StaticMeshVertexBuffer.SetVertexUV(i, 0, Vertex.TextureCoordinate[0]);
-				VertexBuffers.ColorVertexBuffer.VertexColor(i) = Vertex.Color;
-			}
-
-			{
-				auto& VertexBuffer = VertexBuffers.PositionVertexBuffer;
-				void* VertexBufferData = RHICmdList.LockBuffer(VertexBuffer.VertexBufferRHI, 0, VertexBuffer.GetNumVertices() * VertexBuffer.GetStride(), RLM_WriteOnly);
-				FMemory::Memcpy(VertexBufferData, VertexBuffer.GetVertexData(), VertexBuffer.GetNumVertices() * VertexBuffer.GetStride());
-				RHICmdList.UnlockBuffer(VertexBuffer.VertexBufferRHI);
-			}
-
-			{
-				auto& VertexBuffer = VertexBuffers.ColorVertexBuffer;
-				void* VertexBufferData = RHICmdList.LockBuffer(VertexBuffer.VertexBufferRHI, 0, VertexBuffer.GetNumVertices() * VertexBuffer.GetStride(), RLM_WriteOnly);
-				FMemory::Memcpy(VertexBufferData, VertexBuffer.GetVertexData(), VertexBuffer.GetNumVertices() * VertexBuffer.GetStride());
-				RHICmdList.UnlockBuffer(VertexBuffer.VertexBufferRHI);
-			}
-
-			{
-				auto& VertexBuffer = VertexBuffers.StaticMeshVertexBuffer;
-				void* VertexBufferData = RHICmdList.LockBuffer(VertexBuffer.TangentsVertexBuffer.VertexBufferRHI, 0, VertexBuffer.GetTangentSize(), RLM_WriteOnly);
-				FMemory::Memcpy(VertexBufferData, VertexBuffer.GetTangentData(), VertexBuffer.GetTangentSize());
-				RHICmdList.UnlockBuffer(VertexBuffer.TangentsVertexBuffer.VertexBufferRHI);
-			}
-
-			{
-				auto& VertexBuffer = VertexBuffers.StaticMeshVertexBuffer;
-				void* VertexBufferData = RHICmdList.LockBuffer(VertexBuffer.TexCoordVertexBuffer.VertexBufferRHI, 0, VertexBuffer.GetTexCoordSize(), RLM_WriteOnly);
-				FMemory::Memcpy(VertexBufferData, VertexBuffer.GetTexCoordData(), VertexBuffer.GetTexCoordSize());
-				RHICmdList.UnlockBuffer(VertexBuffer.TexCoordVertexBuffer.VertexBufferRHI);
-			}
-
-			void* IndexBufferData = RHICmdList.LockBuffer(IndexBuffer.IndexBufferRHI, 0, Indices.Num() * sizeof(int32), RLM_WriteOnly);
-			FMemory::Memcpy(IndexBufferData, &Indices[0], Indices.Num() * sizeof(int32));
-			RHICmdList.UnlockBuffer(IndexBuffer.IndexBufferRHI);
+			BuildCableMesh(RHICmdList, NewDynamicData);
 
 #if RHI_RAYTRACING
 			if (bSupportRayTracing && CVarRayTracingCableMeshes.GetValueOnRenderThread() != 0)
