[HWRT] Update ray tracing geometry instead of recreating resource when cable dynamic data changes.

diff --git a/Engine/Plugins/Runtime/CableComponent/Source/CableComponent/Private/CableComponent.cpp b/Engine/Plugins/Runtime/CableComponent/Source/CableComponent/Private/CableComponent.cpp
--- a/Engine/Plugins/Runtime/CableComponent/Source/CableComponent/Private/CableComponent.cpp
+++ b/Engine/Plugins/Runtime/CableComponent/Source/CableComponent/Private/CableComponent.cpp
@@ -282,15 +282,10 @@ public:
 			BuildCableMesh(RHICmdList, NewDynamicData);
 
 #if RHI_RAYTRACING
-			if (bSupportRayTracing && CVarRayTracingCableMeshes.GetValueOnRenderThread() != 0)
+			if (StaticRayTracingGeometry.IsInitialized())
 			{
-				FRayTracingGeometry& RayTracingGeometry = StaticRayTracingGeometry;
-				if (RayTracingGeometry.IsValid())
-				{
-					RayTracingGeometry.ReleaseResource();
-					CreateRayTracingGeometry_RenderingThread(RayTracingGeometry, RHICmdList);
-					bNeedsToUpdateRayTracingCache = true;
-				}
+				StaticRayTracingGeometry.SetRequiresBuild(true);
+				bNeedsToUpdateRayTracingCache = true;
 			}
 #endif
 
@@ -473,6 +468,9 @@ public:
 			{
 				return;
 			}
+
+			// TODO: Move this call to SetDynamicData_RenderThread(...) once race condition with FRayTracingGeometryManager::Tick() is addressed.
+			StaticRayTracingGeometry.RequestBuildIfNeeded(Collector.GetRHICommandList(), ERTAccelerationStructureBuildPriority::Normal);
 		}
 
 		FRayTracingGeometry& Geometry = bEvaluateWPO ? DynamicRayTracingGeometry : StaticRayTracingGeometry;
