[HWRT] Re-create evicted cable component RT geometry so it becomes resident.

#rnx

diff --git a/Engine/Plugins/Runtime/CableComponent/Source/CableComponent/Private/CableComponent.cpp b/Engine/Plugins/Runtime/CableComponent/Source/CableComponent/Private/CableComponent.cpp
--- a/Engine/Plugins/Runtime/CableComponent/Source/CableComponent/Private/CableComponent.cpp
+++ b/Engine/Plugins/Runtime/CableComponent/Source/CableComponent/Private/CableComponent.cpp
@@ -470,7 +470,14 @@ public:
 			}
 
 			// TODO: Move this call to SetDynamicData_RenderThread(...) once race condition with FRayTracingGeometryManager::Tick() is addressed.
-			StaticRayTracingGeometry.RequestBuildIfNeeded(Collector.GetRHICommandList(), ERTAccelerationStructureBuildPriority::Normal);
+			if (StaticRayTracingGeometry.IsEvicted())
+			{
+				StaticRayTracingGeometry.CreateRayTracingGeometry(Collector.GetRHICommandList(), ERTAccelerationStructureBuildPriority::Normal);
+			}
+			else
+			{
+				StaticRayTracingGeometry.RequestBuildIfNeeded(Collector.GetRHICommandList(), ERTAccelerationStructureBuildPriority::Normal);
+			}
 		}
 
 		FRayTracingGeometry& Geometry = bEvaluateWPO ? DynamicRayTracingGeometry : StaticRayTracingGeometry;
