[HWRT] Refactored management of StaticRayTracingGeometry in FCableSceneProxy.
- request build in SetDynamicData_RenderThread(...) instead of GetDynamicRayTracingInstances(...).
- mark geometry as referenced and skip adding ray tracing instance when geometry is evicted.

#Synced-CL 52420933

diff --git a/Engine/Plugins/Runtime/CableComponent/Source/CableComponent/Private/CableComponent.cpp b/Engine/Plugins/Runtime/CableComponent/Source/CableComponent/Private/CableComponent.cpp
--- a/Engine/Plugins/Runtime/CableComponent/Source/CableComponent/Private/CableComponent.cpp
+++ b/Engine/Plugins/Runtime/CableComponent/Source/CableComponent/Private/CableComponent.cpp
@@ -122,11 +122,9 @@ public:
 			IndexBuffer.InitResource(RHICmdList);
 
 #if RHI_RAYTRACING
-
 			if (bSupportRayTracing)
 			{
-				FRayTracingGeometry& RayTracingGeometry = StaticRayTracingGeometry;
-				CreateRayTracingGeometry_RenderingThread(RayTracingGeometry, RHICmdList);
+				CreateStaticRayTracingGeometry(RHICmdList);
 				bNeedsToUpdateRayTracingCache = true;
 
 				if (bDynamicRayTracingGeometry)
@@ -134,7 +132,7 @@ public:
 					CreateDynamicRayTracingGeometries(RHICmdList);
 				}
 			}
-#endif
+#endif // RHI_RAYTRACING
 		});
 	}
 
@@ -284,7 +282,16 @@ public:
 #if RHI_RAYTRACING
 			if (StaticRayTracingGeometry.IsInitialized())
 			{
-				StaticRayTracingGeometry.SetRequiresBuild(true);
+				if (StaticRayTracingGeometry.IsEvicted())
+				{
+					StaticRayTracingGeometry.CreateRayTracingGeometry(RHICmdList, ERTAccelerationStructureBuildPriority::Normal);
+				}
+				else if(!StaticRayTracingGeometry.HasPendingBuildRequest())
+				{
+					StaticRayTracingGeometry.SetRequiresBuild(true);
+					StaticRayTracingGeometry.RequestBuildIfNeeded(RHICmdList, ERTAccelerationStructureBuildPriority::Normal);
+				}
+
 				bNeedsToUpdateRayTracingCache = true;
 			}
 #endif
@@ -491,14 +498,10 @@ public:
 				return;
 			}
 
-			// TODO: Move this call to SetDynamicData_RenderThread(...) once race condition with FRayTracingGeometryManager::Tick() is addressed.
 			if (StaticRayTracingGeometry.IsEvicted())
 			{
-				StaticRayTracingGeometry.CreateRayTracingGeometry(Collector.GetRHICommandList(), ERTAccelerationStructureBuildPriority::Normal);
-			}
-			else
-			{
-				StaticRayTracingGeometry.RequestBuildIfNeeded(Collector.GetRHICommandList(), ERTAccelerationStructureBuildPriority::Normal);
+				Collector.AddReferencedGeometry(StaticRayTracingGeometry.GetGeometryHandle());
+				return;
 			}
 		}
 
@@ -604,7 +607,7 @@ public:
 	virtual bool IsRayTracingRelevant() const override { return true; }
 	virtual bool IsRayTracingStaticRelevant() const override { return false; }
 
-	void CreateRayTracingGeometry_RenderingThread(FRayTracingGeometry& RayTracingGeometry, FRHICommandListBase& RHICmdList)
+	void CreateStaticRayTracingGeometry(FRHICommandListBase& RHICmdList)
 	{
 		FRayTracingGeometryInitializer Initializer;
 		static const FName DebugName("FCableSceneProxy");
@@ -622,8 +625,8 @@ public:
 		Segment.MaxVertices = VertexBuffers.PositionVertexBuffer.GetNumVertices();
 		Initializer.Segments.Add(Segment);
 
-		RayTracingGeometry.SetInitializer(Initializer);
-		RayTracingGeometry.InitResource(RHICmdList);
+		StaticRayTracingGeometry.SetInitializer(Initializer);
+		StaticRayTracingGeometry.InitResource(RHICmdList);
 	}
 #endif
 
diff --git a/Engine/Source/Runtime/Engine/Private/Rendering/RayTracingGeometryManager.cpp b/Engine/Source/Runtime/Engine/Private/Rendering/RayTracingGeometryManager.cpp
--- a/Engine/Source/Runtime/Engine/Private/Rendering/RayTracingGeometryManager.cpp
+++ b/Engine/Source/Runtime/Engine/Private/Rendering/RayTracingGeometryManager.cpp
@@ -1988,7 +1988,29 @@ void FRayTracingGeometryManager::AddReferencedGeometry(const FRayTracingGeometry
 	}
 	else
 	{
-		ensureMsgf(ReferencedGeometries.IsEmpty(), TEXT("Should only track ReferencedGeometries when using using reference based residency"));
+		ensureMsgf(ReferencedGeometries.IsEmpty(), TEXT("Should only track ReferencedGeometries when using reference based residency"));
+	}
+}
+
+void FRayTracingGeometryManager::AddReferencedGeometries(const TSet<RayTracing::FGeometryHandle>& Geometries)
+{
+	check(IsInRenderingThread() || IsInParallelRenderingThread());
+
+	if (IsRayTracingUsingReferenceBasedResidency())
+	{
+#if DO_CHECK
+		// try to use handle here in development builds to detect handles as early as possible
+		for (RayTracing::FGeometryHandle Handle : Geometries)
+		{
+			GetRegisteredGeometry(Handle);
+		}
+#endif
+
+		ReferencedGeometries.Append(Geometries);
+	}
+	else
+	{
+		ensureMsgf(Geometries.IsEmpty(), TEXT("Should only track ReferencedGeometries when using reference based residency"));
 	}
 }
 
@@ -2010,7 +2032,7 @@ void FRayTracingGeometryManager::AddReferencedGeometryGroups(const TSet<RayTraci
 	}
 	else
 	{
-		ensureMsgf(GeometryGroups.IsEmpty(), TEXT("Should only track ReferencedGeometryGroups when using using reference based residency"));
+		ensureMsgf(GeometryGroups.IsEmpty(), TEXT("Should only track ReferencedGeometryGroups when using reference based residency"));
 	}
 }
 
@@ -2032,7 +2054,7 @@ void FRayTracingGeometryManager::AddReferencedGeometryGroupsForDynamicUpdate(con
 	}
 	else
 	{
-		ensureMsgf(GeometryGroups.IsEmpty(), TEXT("Should only track ReferencedGeometryGroupsForDynamic when using using reference based residency"));
+		ensureMsgf(GeometryGroups.IsEmpty(), TEXT("Should only track ReferencedGeometryGroupsForDynamic when using reference based residency"));
 	}
 }
 
diff --git a/Engine/Source/Runtime/Engine/Private/SceneManagement.cpp b/Engine/Source/Runtime/Engine/Private/SceneManagement.cpp
--- a/Engine/Source/Runtime/Engine/Private/SceneManagement.cpp
+++ b/Engine/Source/Runtime/Engine/Private/SceneManagement.cpp
@@ -366,9 +366,9 @@ FMeshBatchAndRelevance::FMeshBatchAndRelevance(const FMeshBatch& InMesh, const F
 FRayTracingInstanceCollector::FRayTracingInstanceCollector(
 	ERHIFeatureLevel::Type InFeatureLevel,
 	FSceneRenderingBulkObjectAllocator& InBulkAllocator,
-	bool bInTrackReferencedGeometryGroups)
+	bool bInTrackReferencedGeometry)
 	: FMeshElementCollector(InFeatureLevel, InBulkAllocator, FMeshElementCollector::ECommitFlags::DeferAll)
-	, bTrackReferencedGeometryGroups(bInTrackReferencedGeometryGroups)
+	, bTrackReferencedGeometry(bInTrackReferencedGeometry)
 {
 }
 
@@ -409,9 +409,17 @@ void FRayTracingInstanceCollector::AddRayTracingInstance(int32 ViewIndex, FRayTr
 	RayTracingInstances.AddElement(MoveTemp(Tmp));
 }
 
+void FRayTracingInstanceCollector::AddReferencedGeometry(RayTracing::FGeometryHandle Geometry)
+{
+	if (bTrackReferencedGeometry)
+	{
+		ReferencedGeometries.Add(Geometry);
+	}
+}
+
 void FRayTracingInstanceCollector::AddReferencedGeometryGroup(RayTracing::FGeometryGroupHandle GeometryGroup)
 {
-	if (bTrackReferencedGeometryGroups)
+	if (bTrackReferencedGeometry)
 	{
 		ReferencedGeometryGroups.Add(GeometryGroup);
 	}
@@ -419,7 +427,7 @@ void FRayTracingInstanceCollector::AddReferencedGeometryGroup(RayTracing::FGeome
 
 void FRayTracingInstanceCollector::AddReferencedGeometryGroupForDynamicUpdate(RayTracing::FGeometryGroupHandle GeometryGroup)
 {
-	if (bTrackReferencedGeometryGroups)
+	if (bTrackReferencedGeometry)
 	{
 		ReferencedGeometryGroupsForDynamicUpdate.Add(GeometryGroup);
 	}
diff --git a/Engine/Source/Runtime/Engine/Public/Rendering/RayTracingGeometryManager.h b/Engine/Source/Runtime/Engine/Public/Rendering/RayTracingGeometryManager.h
--- a/Engine/Source/Runtime/Engine/Public/Rendering/RayTracingGeometryManager.h
+++ b/Engine/Source/Runtime/Engine/Public/Rendering/RayTracingGeometryManager.h
@@ -50,6 +50,7 @@ public:
 	ENGINE_API virtual void RequestUpdateCachedRenderState(RayTracing::FGeometryGroupHandle InRayTracingGeometryGroupHandle) override;
 
 	ENGINE_API void AddReferencedGeometry(const FRayTracingGeometry* Geometry);
+	ENGINE_API void AddReferencedGeometries(const TSet<RayTracing::FGeometryHandle>& Geometries);
 	ENGINE_API void AddReferencedGeometryGroups(const TSet<RayTracing::FGeometryGroupHandle>& GeometryGroups);
 	ENGINE_API void AddReferencedGeometryGroupsForDynamicUpdate(const TSet<RayTracing::FGeometryGroupHandle>& GeometryGroups);
 
diff --git a/Engine/Source/Runtime/Engine/Public/SceneManagement.h b/Engine/Source/Runtime/Engine/Public/SceneManagement.h
--- a/Engine/Source/Runtime/Engine/Public/SceneManagement.h
+++ b/Engine/Source/Runtime/Engine/Public/SceneManagement.h
@@ -1523,8 +1523,8 @@ public:
 
 	ENGINE_API void AddRayTracingInstance(int32 ViewIndex, FRayTracingInstance Instance);
 
+	ENGINE_API void AddReferencedGeometry(RayTracing::FGeometryHandle Geometry);
 	ENGINE_API void AddReferencedGeometryGroup(RayTracing::FGeometryGroupHandle GeometryGroup);
-
 	ENGINE_API void AddReferencedGeometryGroupForDynamicUpdate(RayTracing::FGeometryGroupHandle GeometryGroup);
 
 	UE_DEPRECATED(5.7, "Please specify ViewIndex.")
@@ -1548,7 +1548,7 @@ private:
 	ENGINE_API FRayTracingInstanceCollector(
 		ERHIFeatureLevel::Type InFeatureLevel,
 		FSceneRenderingBulkObjectAllocator& InBulkAllocator,
-		bool bInTrackReferencedGeometryGroups);
+		bool bInTrackReferencedGeometry);
 
 	ENGINE_API void AddViewMeshArrays(
 		const FSceneView* InView,
@@ -1573,10 +1573,12 @@ private:
 	TArray<TSet<FRDGPooledBuffer*>, TInlineAllocator<2, SceneRenderingAllocator> > RDGPooledBuffers;
 	TChunkedArray<FRayTracingInstanceAndViewIndex, 16384, SceneRenderingAllocator> RayTracingInstances;
 	TChunkedArray<FRayTracingDynamicGeometryUpdateRequest, 16384, SceneRenderingAllocator> RayTracingGeometriesToUpdate;
+
+	TSet<RayTracing::FGeometryHandle> ReferencedGeometries;
 	TSet<RayTracing::FGeometryGroupHandle> ReferencedGeometryGroups;
 	TSet<RayTracing::FGeometryGroupHandle> ReferencedGeometryGroupsForDynamicUpdate;
 
-	const bool bTrackReferencedGeometryGroups;
+	const bool bTrackReferencedGeometry;
 
 	uint32 VisibilityMap;
 
diff --git a/Engine/Source/Runtime/Renderer/Private/RayTracing/RayTracing.cpp b/Engine/Source/Runtime/Renderer/Private/RayTracing/RayTracing.cpp
--- a/Engine/Source/Runtime/Renderer/Private/RayTracing/RayTracing.cpp
+++ b/Engine/Source/Runtime/Renderer/Private/RayTracing/RayTracing.cpp
@@ -425,7 +425,7 @@ namespace RayTracing
 		TConstArrayView<const FSceneOptions*> ViewSceneOptions;
 		FSceneRenderingBulkObjectAllocator& BulkAllocator;
 
-		const bool bTrackReferencedGeometryGroups;
+		const bool bTrackReferencedGeometry;
 		bool bAnyViewRequiresTranslucentGeometry;
 
 		FRHICommandList* RHICmdList;
@@ -777,11 +777,11 @@ namespace RayTracing
 		, Views(InViews)
 		, ViewSceneOptions(InViewSceneOptions)
 		, BulkAllocator(InBulkAllocator)
-		, bTrackReferencedGeometryGroups(IsRayTracingUsingReferenceBasedResidency())
+		, bTrackReferencedGeometry(IsRayTracingUsingReferenceBasedResidency())
 		, RHICmdList(new FRHICommandList(FRHIGPUMask::All()))
 		, DynamicVertexBuffer(*RHICmdList)
 		, DynamicIndexBuffer(*RHICmdList)
-		, RayTracingInstanceCollector(Scene.GetFeatureLevel(), InBulkAllocator, bTrackReferencedGeometryGroups)
+		, RayTracingInstanceCollector(Scene.GetFeatureLevel(), InBulkAllocator, bTrackReferencedGeometry)
 	{
 		RHICmdList->SwitchPipeline(ERHIPipeline::Graphics);
 
@@ -908,17 +908,20 @@ namespace RayTracing
 
 		RayTracingInstanceCollector.RayTracingGeometriesToUpdate.Empty();
 
-		if (bTrackReferencedGeometryGroups)
+		if (bTrackReferencedGeometry)
 		{
 			// TODO: Could run in parallel thread if properly synchronized with static ray tracing instances tasks
+			((FRayTracingGeometryManager*)GRayTracingGeometryManager)->AddReferencedGeometries(RayTracingInstanceCollector.ReferencedGeometries);
 			((FRayTracingGeometryManager*)GRayTracingGeometryManager)->AddReferencedGeometryGroups(RayTracingInstanceCollector.ReferencedGeometryGroups);
 			((FRayTracingGeometryManager*)GRayTracingGeometryManager)->AddReferencedGeometryGroupsForDynamicUpdate(RayTracingInstanceCollector.ReferencedGeometryGroupsForDynamicUpdate);
 
+			RayTracingInstanceCollector.ReferencedGeometries.Empty();
 			RayTracingInstanceCollector.ReferencedGeometryGroups.Empty();
 			RayTracingInstanceCollector.ReferencedGeometryGroupsForDynamicUpdate.Empty();
 		}
 		else
 		{
+			check(RayTracingInstanceCollector.ReferencedGeometries.IsEmpty());
 			check(RayTracingInstanceCollector.ReferencedGeometryGroups.IsEmpty());
 			check(RayTracingInstanceCollector.ReferencedGeometryGroupsForDynamicUpdate.IsEmpty());
 		}
